Merging using Main-Branch-Map

#UE4 #PerfCounters
- Changed the stats port format to be "HTTP-like" (just enough to be compatible with curl/wget/etc)
- stubbed out a super basic framework for being able to issue commands via this port.

diff --git a/Engine/Source/Runtime/PerfCounters/Private/PerfCounters.cpp b/Engine/Source/Runtime/PerfCounters/Private/PerfCounters.cpp
--- a/Engine/Source/Runtime/PerfCounters/Private/PerfCounters.cpp
+++ b/Engine/Source/Runtime/PerfCounters/Private/PerfCounters.cpp
@@ -122,6 +122,13 @@ FString FPerfCounters::ToJson() const
 	return JsonStr;
 }
 
+static bool SendAsUtf8(FSocket* Conn, const FString& Message)
+{
+	FTCHARToUTF8 ConvertToUtf8(*Message);
+	int32 BytesSent = 0;
+	return Conn->Send(reinterpret_cast<const uint8*>(ConvertToUtf8.Get()), ConvertToUtf8.Length(), BytesSent) && BytesSent == ConvertToUtf8.Length();
+}
+
 bool FPerfCounters::Tick(float DeltaTime)
 {
 	// if we didn't get a socket, don't tick
@@ -135,14 +142,63 @@ bool FPerfCounters::Tick(float DeltaTime)
 	FSocket* IncomingConnection = Socket->Accept(PerfCounterRequest);
 	if (IncomingConnection)
 	{
-		// handle the connection
-		int32 BytesSent = 0;
-		FTCHARToUTF8 ConvertToUtf8(*ToJson());
-		bool bSuccess = IncomingConnection->Send(reinterpret_cast<const uint8*>(ConvertToUtf8.Get()), ConvertToUtf8.Length(), BytesSent);
-		if (!bSuccess || BytesSent != ConvertToUtf8.Length())
+		// make sure this is non-blocking
+		bool bSuccess = false;
+		IncomingConnection->SetNonBlocking(true);
+
+		// read any data that's ready
+		// NOTE: this is not a full HTTP implementation, just enough to be usable by curl
+		uint8 Buffer[2*1024] = { 0 };
+		int32 DataLen = 0;
+		if (IncomingConnection->Recv(Buffer, sizeof(Buffer)-1, DataLen, ESocketReceiveFlags::None))
+		{
+			// scan the buffer for a line
+			FUTF8ToTCHAR WideBuffer(reinterpret_cast<const ANSICHAR*>(Buffer));
+			const TCHAR* BufferEnd = FCString::Strstr(WideBuffer.Get(), TEXT("\r\n"));
+			if (BufferEnd != nullptr)
+			{
+				// crack into pieces
+				FString MainLine(BufferEnd - WideBuffer.Get(), WideBuffer.Get());
+				TArray<FString> Tokens;
+				MainLine.ParseIntoArrayWS(Tokens);
+				if (Tokens.Num() >= 2)
+				{
+					FString Body;
+					int ResponseCode = 200;
+
+					// handle the request
+					if (Tokens[0] != TEXT("GET"))
+					{
+						Body = FString::Printf(TEXT("{ \"error\": \"Method %s not allowed\" }"), *Tokens[0]);
+						ResponseCode = 405;
+					}
+					else if (Tokens[1] == TEXT("/stats"))
+					{
+						Body = ToJson();
+					}
+					else
+					{
+						Body = FString::Printf(TEXT("{ \"error\": \"%s not found\" }"), *Tokens[1]);
+						ResponseCode = 404;
+					}
+
+					// send the response headers
+					FString Header = FString::Printf(TEXT("HTTP/1.0 %d\r\nContent-Length: %d\r\nContent-Type: application/json\r\n\r\n"), ResponseCode, Body.Len());
+					if (SendAsUtf8(IncomingConnection, Header) && SendAsUtf8(IncomingConnection, Body))
+					{
+						bSuccess = true;
+					}
+				}
+			}
+		}
+
+		// log if we didn't succeed
+		if (!bSuccess)
 		{
 			UE_LOG(LogPerfCounters, Warning, TEXT("FPerfCounters was unable to send a JSON response (or sent partial response)"));
 		}
+
+		// close the socket (whether we processed or not
 		IncomingConnection->Close();
 		ISocketSubsystem::Get(PLATFORM_SOCKETSUBSYSTEM)->DestroySocket(IncomingConnection);
 	}
