Merging using Main-Branch-Map CL 2628904
#PerfCounters #UE4: Add more error logging to HTTP-like request parsing

diff --git a/Engine/Source/Runtime/PerfCounters/Private/PerfCounters.cpp b/Engine/Source/Runtime/PerfCounters/Private/PerfCounters.cpp
--- a/Engine/Source/Runtime/PerfCounters/Private/PerfCounters.cpp
+++ b/Engine/Source/Runtime/PerfCounters/Private/PerfCounters.cpp
@@ -184,12 +184,35 @@ bool FPerfCounters::Tick(float DeltaTime)
 
 					// send the response headers
 					FString Header = FString::Printf(TEXT("HTTP/1.0 %d\r\nContent-Length: %d\r\nContent-Type: application/json\r\n\r\n"), ResponseCode, Body.Len());
-					if (SendAsUtf8(IncomingConnection, Header) && SendAsUtf8(IncomingConnection, Body))
+					if (SendAsUtf8(IncomingConnection, Header))
 					{
-						bSuccess = true;
+						if (SendAsUtf8(IncomingConnection, Body))
+						{
+							bSuccess = true;
+						}
+						else
+						{
+							UE_LOG(LogPerfCounters, Warning, TEXT("Unable to send full HTTP response body"));
+						}
 					}
+					else
+					{
+						UE_LOG(LogPerfCounters, Warning, TEXT("Unable to send HTTP response header: %s"), *Header);
+					}
+				}
+				else
+				{
+					UE_LOG(LogPerfCounters, Warning, TEXT("Unable to parse HTTP request header: %s"), *MainLine);
 				}
 			}
+			else
+			{
+				UE_LOG(LogPerfCounters, Warning, TEXT("Unable to immediately receive full request header"));
+			}
+		}
+		else
+		{
+			UE_LOG(LogPerfCounters, Warning, TEXT("Unable to immediately receive request header"));
 		}
 
 		// log if we didn't succeed
