PerfCounter updates
- new "exec" endpoint
- refactored tick code
- added new process request code to fix issues where the socket wasn't ready

diff --git a/Engine/Source/Runtime/PerfCounters/Private/PerfCounters.cpp b/Engine/Source/Runtime/PerfCounters/Private/PerfCounters.cpp
--- a/Engine/Source/Runtime/PerfCounters/Private/PerfCounters.cpp
+++ b/Engine/Source/Runtime/PerfCounters/Private/PerfCounters.cpp
@@ -142,94 +142,157 @@ bool FPerfCounters::Tick(float DeltaTime)
 	FSocket* IncomingConnection = Socket->Accept(PerfCounterRequest);
 	if (IncomingConnection)
 	{
+		if (0)
+		{
+			TSharedRef<FInternetAddr> FromAddr = ISocketSubsystem::Get(PLATFORM_SOCKETSUBSYSTEM)->CreateInternetAddr();
+			IncomingConnection->GetPeerAddress(*FromAddr);
+			UE_LOG(LogPerfCounters, Log, TEXT("New connection from %s"), *FromAddr->ToString(true));
+		}
+
 		// make sure this is non-blocking
-		bool bSuccess = false;
 		IncomingConnection->SetNonBlocking(true);
 
-		// read any data that's ready
-		// NOTE: this is not a full HTTP implementation, just enough to be usable by curl
-		uint8 Buffer[2*1024] = { 0 };
-		int32 DataLen = 0;
-		if (IncomingConnection->Recv(Buffer, sizeof(Buffer)-1, DataLen, ESocketReceiveFlags::None))
+		new (Connections) FPerfConnection(IncomingConnection);
+	}
+
+	TArray<FPerfConnection> ConnectionsToClose;
+	for (FPerfConnection& Connection : Connections)
+	{
+		FSocket* ExistingSocket = Connection.Connection;
+		if (ExistingSocket && ExistingSocket->Wait(ESocketWaitConditions::WaitForRead, FTimespan::Zero()))
 		{
-			// scan the buffer for a line
-			FUTF8ToTCHAR WideBuffer(reinterpret_cast<const ANSICHAR*>(Buffer));
-			const TCHAR* BufferEnd = FCString::Strstr(WideBuffer.Get(), TEXT("\r\n"));
-			if (BufferEnd != nullptr)
+			// read any data that's ready
+			// NOTE: this is not a full HTTP implementation, just enough to be usable by curl
+			uint8 Buffer[2 * 1024] = { 0 };
+			int32 DataLen = 0;
+			if (ExistingSocket->Recv(Buffer, sizeof(Buffer) - 1, DataLen, ESocketReceiveFlags::None))
 			{
-				// crack into pieces
-				FString MainLine(BufferEnd - WideBuffer.Get(), WideBuffer.Get());
-				TArray<FString> Tokens;
-				MainLine.ParseIntoArrayWS(Tokens);
-				if (Tokens.Num() >= 2)
+				FResponse Response;
+				if (ProcessRequest(Buffer, DataLen, Response))
 				{
-					FString Body;
-					int ResponseCode = 200;
-
-					// handle the request
-					if (Tokens[0] != TEXT("GET"))
-					{
-						Body = FString::Printf(TEXT("{ \"error\": \"Method %s not allowed\" }"), *Tokens[0]);
-						ResponseCode = 405;
-					}
-					else if (Tokens[1] == TEXT("/stats"))
-					{
-						Body = ToJson();
-					}
-					else
+					if (SendAsUtf8(ExistingSocket, Response.Header))
 					{
-						Body = FString::Printf(TEXT("{ \"error\": \"%s not found\" }"), *Tokens[1]);
-						ResponseCode = 404;
-					}
-
-					// send the response headers
-					FString Header = FString::Printf(TEXT("HTTP/1.0 %d\r\nContent-Length: %d\r\nContent-Type: application/json\r\n\r\n"), ResponseCode, Body.Len());
-					if (SendAsUtf8(IncomingConnection, Header))
-					{
-						if (SendAsUtf8(IncomingConnection, Body))
-						{
-							bSuccess = true;
-						}
-						else
+						if (!SendAsUtf8(ExistingSocket, Response.Body))
 						{
 							UE_LOG(LogPerfCounters, Warning, TEXT("Unable to send full HTTP response body"));
 						}
 					}
 					else
 					{
-						UE_LOG(LogPerfCounters, Warning, TEXT("Unable to send HTTP response header: %s"), *Header);
+						UE_LOG(LogPerfCounters, Warning, TEXT("Unable to send HTTP response header: %s"), *Response.Header);
 					}
 				}
-				else
-				{
-					UE_LOG(LogPerfCounters, Warning, TEXT("Unable to parse HTTP request header: %s"), *MainLine);
-				}
 			}
 			else
 			{
-				UE_LOG(LogPerfCounters, Warning, TEXT("Unable to immediately receive full request header"));
+				UE_LOG(LogPerfCounters, Warning, TEXT("Unable to immediately receive request header"));
 			}
+
+			ConnectionsToClose.Add(Connection);
 		}
-		else
+		else if (Connection.ElapsedTime > 5.0f)
 		{
-			UE_LOG(LogPerfCounters, Warning, TEXT("Unable to immediately receive request header"));
+			ConnectionsToClose.Add(Connection);
 		}
 
-		// log if we didn't succeed
-		if (!bSuccess)
+		Connection.ElapsedTime += DeltaTime;
+	}
+
+	for (FPerfConnection& Connection : ConnectionsToClose)
+	{
+		Connections.RemoveSingleSwap(Connection);
+
+		FSocket* ClosingSocket = Connection.Connection;
+		if (ClosingSocket)
 		{
-			UE_LOG(LogPerfCounters, Warning, TEXT("FPerfCounters was unable to send a JSON response (or sent partial response)"));
+			// close the socket (whether we processed or not)
+			ClosingSocket->Close();
+			ISocketSubsystem::Get(PLATFORM_SOCKETSUBSYSTEM)->DestroySocket(ClosingSocket);
+			
+			if (0)
+			{
+				UE_LOG(LogPerfCounters, Log, TEXT("Closed connection."));
+			}
 		}
-
-		// close the socket (whether we processed or not
-		IncomingConnection->Close();
-		ISocketSubsystem::Get(PLATFORM_SOCKETSUBSYSTEM)->DestroySocket(IncomingConnection);
 	}
 
 	// keep ticking
 	return true;
 }
 
+bool FPerfCounters::ProcessRequest(uint8* Buffer, int32 BufferLen, FResponse& Response)
+{
+	bool bSuccess = false;
+
+	// scan the buffer for a line
+	FUTF8ToTCHAR WideBuffer(reinterpret_cast<const ANSICHAR*>(Buffer));
+	const TCHAR* BufferEnd = FCString::Strstr(WideBuffer.Get(), TEXT("\r\n"));
+	if (BufferEnd != nullptr)
+	{
+		// crack into pieces
+		FString MainLine(BufferEnd - WideBuffer.Get(), WideBuffer.Get());
+		TArray<FString> Tokens;
+		MainLine.ParseIntoArrayWS(Tokens);
+		if (Tokens.Num() >= 2)
+		{
+			FString ContentType(TEXT("application/json"));
+			Response.Code = 200;
+
+			// handle the request
+			if (Tokens[0] != TEXT("GET"))
+			{
+				Response.Body = FString::Printf(TEXT("{ \"error\": \"Method %s not allowed\" }"), *Tokens[0]);
+				Response.Code = 405;
+			}
+			else if (Tokens[1] == TEXT("/stats"))
+			{
+				Response.Body = ToJson();
+			}
+			else if (Tokens[1].StartsWith(TEXT("/exec?c=")))
+			{
+				FString ExecCmd = Tokens[1].Mid(8);
+				FString ExecCmdDecoded = FGenericPlatformHttp::UrlDecode(ExecCmd);
+
+				FStringOutputDevice StringOutDevice;
+				StringOutDevice.SetAutoEmitLineTerminator(true);
+
+				bool bResult = false;
+				if (ExecCmdCallback.IsBound())
+				{
+					bResult = ExecCmdCallback.Execute(ExecCmdDecoded, StringOutDevice);
+					Response.Body = StringOutDevice;
+					ContentType = TEXT("text/text");
+				}
+				else
+				{
+					Response.Body = FString::Printf(TEXT("{ \"error\": \"exec handler not found\" }"));
+				}
+
+				Response.Code = bResult ? 200 : 404;
+			}
+			else
+			{
+				Response.Body = FString::Printf(TEXT("{ \"error\": \"%s not found\" }"), *Tokens[1]);
+				Response.Code = 404;
+			}
+
+			// send the response headers
+			Response.Header = FString::Printf(TEXT("HTTP/1.0 %d\r\nContent-Length: %d\r\nContent-Type: %s\r\n\r\n"), Response.Code, Response.Body.Len(), *ContentType);
+			bSuccess = true;
+		}
+		else
+		{
+			UE_LOG(LogPerfCounters, Warning, TEXT("Unable to parse HTTP request header: %s"), *MainLine);
+		}
+	}
+	else
+	{
+		UE_LOG(LogPerfCounters, Warning, TEXT("Unable to immediately receive full request header"));
+	}
+
+	return bSuccess;
+}
+
 void FPerfCounters::SetNumber(const FString& Name, double Value) 
 {
 	FJsonVariant& JsonValue = PerfCounterMap.FindOrAdd(Name);
@@ -244,9 +307,9 @@ void FPerfCounters::SetString(const FString& Name, const FString& Value)
 	JsonValue.StringValue = Value;
 }
 
-void FPerfCounters::SetJson(const FString& Name, const FProduceJsonCounterValue& Callback)
+void FPerfCounters::SetJson(const FString& Name, const FProduceJsonCounterValue& InCallback)
 {
 	FJsonVariant& JsonValue = PerfCounterMap.FindOrAdd(Name);
 	JsonValue.Format = FJsonVariant::Callback;
-	JsonValue.CallbackValue = Callback;
+	JsonValue.CallbackValue = InCallback;
 }
diff --git a/Engine/Source/Runtime/PerfCounters/Private/PerfCounters.h b/Engine/Source/Runtime/PerfCounters/Private/PerfCounters.h
--- a/Engine/Source/Runtime/PerfCounters/Private/PerfCounters.h
+++ b/Engine/Source/Runtime/PerfCounters/Private/PerfCounters.h
@@ -28,7 +28,8 @@ public:
 	const FString& GetInstanceName() const override { return UniqueInstanceId; }
 	virtual void SetNumber(const FString& Name, double Value) override;
 	virtual void SetString(const FString& Name, const FString& Value) override;
-	virtual void SetJson(const FString& Name, const FProduceJsonCounterValue& Callback) override;
+	virtual void SetJson(const FString& Name, const FProduceJsonCounterValue& InCallback) override;
+	virtual FPerfCounterExecCommandCallback& OnPerfCounterExecCommand() override { return ExecCmdCallback; }
 	// IPerfCounters interface end
 
 private:
@@ -36,6 +37,64 @@ private:
 	FString ToJson() const;
 
 private:
+	
+	/**
+	 * Simple response structure for returning output to requestor
+	 */
+	struct FResponse
+	{
+		/** http header */
+		FString Header;
+		/** http body */
+		FString Body;
+		/** http response code */
+		int32 Code;
+
+		FResponse() :
+			Code(0)
+		{}
+	};
+	
+	/**
+	 * Simple connection structure for keeping track of incoming/active connections
+	 */
+	struct FPerfConnection
+	{
+		/** accepted external socket */
+		FSocket* Connection;
+		/** time connection has existed */
+		float ElapsedTime;
+
+		FPerfConnection() :
+			Connection(nullptr),
+			ElapsedTime(0.0f)
+		{}
+
+		FPerfConnection(FSocket* InConnection) :
+			Connection(InConnection),
+			ElapsedTime(0.0f)
+		{}
+
+		bool operator==(const FPerfConnection& Other) const
+		{
+			return Connection == Other.Connection;
+		}
+	};
+
+	/** all active connections */
+	TArray<FPerfConnection> Connections;
+
+	/**
+	 * Process the incoming request from an active socket
+	 *
+	 * @param Buffer data sent from the requestor
+	 * @param BufferLen size of the data sent
+	 * @param Response [out] response to be given to the requestor
+	 *
+	 * @return true if the response is valid, false otherwise
+	 */
+	bool ProcessRequest(uint8* Buffer, int32 BufferLen, FResponse& Response);
+
 	struct FJsonVariant
 	{
 		enum { Null, String, Number, Callback } Format;
@@ -52,6 +111,9 @@ private:
 	/** Map of all known performance counters */
 	TMap<FString, FJsonVariant>  PerfCounterMap;
 
+	/** Bound callback for script command execution */
+	FPerfCounterExecCommandCallback ExecCmdCallback;
+
 	/* Listen socket for outputting JSON on request */
 	FSocket* Socket;
 };
diff --git a/Engine/Source/Runtime/PerfCounters/Public/PerfCountersModule.h b/Engine/Source/Runtime/PerfCounters/Public/PerfCountersModule.h
--- a/Engine/Source/Runtime/PerfCounters/Public/PerfCountersModule.h
+++ b/Engine/Source/Runtime/PerfCounters/Public/PerfCountersModule.h
@@ -10,8 +10,23 @@ template <class CharType, class PrintPolicy>
 class TJsonWriter;
 typedef TSharedRef< TJsonWriter<TCHAR,TPrettyJsonPrintPolicy<TCHAR> > > FPrettyJsonWriter;
 
+/**
+ * Delegate called for a given counter to generate custom json at the time the query is made
+ *
+ * @param JsonWriter json output writer
+ */
 DECLARE_DELEGATE_OneParam(FProduceJsonCounterValue, const FPrettyJsonWriter& /* JsonWriter */);
 
+/**
+ * Delegate called when an exec command has been passed in on via the query port
+ *
+ * @param ExecCmd string command to execute
+ * @param Output output device to write the response back to the querying service
+ * 
+ * @return true if the call was handled, false otherwise
+ */
+DECLARE_DELEGATE_RetVal_TwoParams(bool, FPerfCounterExecCommandCallback, const FString& /*ExecCmd*/, FOutputDevice& /*Output*/);
+
 /**
  * A programming interface for setting/updating performance counters
  */
@@ -33,6 +48,9 @@ public:
 	/** Make a callback so we can request more extensive types on demand (presumably backed by some struct locally) */
 	virtual void SetJson(const FString& Name, const FProduceJsonCounterValue& Callback) = 0;
 
+	/** @return delegate called when an exec command is to be executed */
+	virtual FPerfCounterExecCommandCallback& OnPerfCounterExecCommand() = 0;
+
 public:
 	/** Set overloads (use these) */
 
