Fix Mac NonUnity compile failure after Orion merge, due to #pragma once in .cpp file.

diff --git a/Engine/Source/Runtime/PerfCounters/Private/ZeroLoad.cpp b/Engine/Source/Runtime/PerfCounters/Private/ZeroLoad.cpp
--- a/Engine/Source/Runtime/PerfCounters/Private/ZeroLoad.cpp
+++ b/Engine/Source/Runtime/PerfCounters/Private/ZeroLoad.cpp
@@ -1,7 +1,5 @@
 // Copyright 1998-2016 Epic Games, Inc. All Rights Reserved.
 
-#pragma once
-
 #include "PerfCounters.h"
 #include "ZeroLoad.h"
 
