Copying //UE4/Orion-Staging to //UE4/Main (Source: //Orion/Dev-General @ 2997507)

diff --git a/Engine/Source/Runtime/PerfCounters/Private/PerfCountersModule.cpp b/Engine/Source/Runtime/PerfCounters/Private/PerfCountersModule.cpp
--- a/Engine/Source/Runtime/PerfCounters/Private/PerfCountersModule.cpp
+++ b/Engine/Source/Runtime/PerfCounters/Private/PerfCountersModule.cpp
@@ -71,125 +71,7 @@ DEFINE_LOG_CATEGORY(LogPerfCounters);
 
 const FName IPerfCounters::Histograms::FrameTime(TEXT("FrameTime"));
 const FName IPerfCounters::Histograms::FrameTimePeriodic(TEXT("FrameTimePeriodic"));
+const FName IPerfCounters::Histograms::FrameTimeWithoutSleep(TEXT("FrameTimeWithoutSleep"));
 const FName IPerfCounters::Histograms::ServerReplicateActorsTime(TEXT("ServerReplicateActorsTime"));
 const FName IPerfCounters::Histograms::SleepTime(TEXT("SleepTime"));
-const FName IPerfCounters::Histograms::ZeroLoadFrameTime(TEXT("ZeroLoadFrameTime"));
-
-void FHistogram::InitLinear(double MinTime, double MaxTime, double BinSize)
-{
-	double CurrentBinMin = MinTime;
-	while (CurrentBinMin < MaxTime)
-	{
-		Bins.Add(FBin(CurrentBinMin, CurrentBinMin + BinSize));
-		CurrentBinMin += BinSize;
-	}
-	Bins.Add(FBin(CurrentBinMin));	// catch-all-above bin
-}
-
-void FHistogram::InitHitchTracking()
-{
-	Bins.Empty();
-
-	Bins.Add(FBin(0.0, 9.0));		// >= 120 fps
-	Bins.Add(FBin(9.0, 17.0));		// 60 - 120 fps
-	Bins.Add(FBin(17.0, 34.0));		// 30 - 60 fps
-	Bins.Add(FBin(34.0, 50.0));		// 20 - 30 fps
-	Bins.Add(FBin(50.0, 67.0));		// 15 - 20 fps
-	Bins.Add(FBin(67.0, 100.0));	// 10 - 15 fps
-	Bins.Add(FBin(100.0, 200.0));	// 5 - 10 fps
-	Bins.Add(FBin(200.0, 300.0));	// < 5 fps
-	Bins.Add(FBin(300.0, 500.0));
-	Bins.Add(FBin(500.0, 750.0));
-	Bins.Add(FBin(750.0, 1000.0));
-	Bins.Add(FBin(1000.0, 1500.0));
-	Bins.Add(FBin(1500.0, 2000.0));
-	Bins.Add(FBin(2000.0, 2500.0));
-	Bins.Add(FBin(2500.0, 5000.0));
-	Bins.Add(FBin(5000.0));
-}
-
-void FHistogram::Reset()
-{
-	for (FBin& Bin : Bins)
-	{
-		Bin.Count = 0;
-		Bin.Sum = 0.0;
-	}
-}
-
-void FHistogram::AddMeasurement(double Value)
-{
-	if (LIKELY(Bins.Num()))
-	{
-		FBin& FirstBin = Bins[0];
-		if (UNLIKELY(Value < FirstBin.MinValue))
-		{
-			return;
-		}
-
-		for (int BinIdx = 0, LastBinIdx = Bins.Num() - 1; BinIdx < LastBinIdx; ++BinIdx)
-		{
-			FBin& Bin = Bins[BinIdx];
-			if (Bin.UpperBound > Value)
-			{
-				++Bin.Count;
-				Bin.Sum += Value;
-				return;
-			}
-		}
-
-		// if we got here, Value did not fit in any of the previous bins
-		FBin& LastBin = Bins.Last();
-		++LastBin.Count;
-		LastBin.Sum += Value;
-	}
-}
-
-void FHistogram::DumpToAnalytics(const FString& ParamNamePrefix, TArray<TPair<FString, double>>& OutParamArray)
-{
-	if (LIKELY(Bins.Num()))
-	{
-		for (int BinIdx = 0, LastBinIdx = Bins.Num() - 1; BinIdx < LastBinIdx; ++BinIdx)
-		{
-			FBin& Bin = Bins[BinIdx];
-			FString ParamName = FString::Printf(TEXT("_%.0f_%.0f"), Bin.MinValue, Bin.UpperBound);
-			OutParamArray.Add(TPairInitializer<FString, double>(ParamNamePrefix + ParamName + TEXT("_Count"), Bin.Count));
-			OutParamArray.Add(TPairInitializer<FString, double>(ParamNamePrefix + ParamName + TEXT("_Sum"), Bin.Sum));
-		}
-
-		FBin& LastBin = Bins.Last();
-		FString ParamName = FString::Printf(TEXT("_%.0f_AndAbove"), LastBin.MinValue);
-		OutParamArray.Add(TPairInitializer<FString, double>(ParamNamePrefix + ParamName + TEXT("_Count"), LastBin.Count));
-		OutParamArray.Add(TPairInitializer<FString, double>(ParamNamePrefix + ParamName + TEXT("_Sum"), LastBin.Sum));
-	}
-}
-
-void FHistogram::DumpToLog(const FString& HistogramName)
-{
-	UE_LOG(LogPerfCounters, Log, TEXT("Histogram '%s': %d bins"), *HistogramName, Bins.Num());
-
-	if (LIKELY(Bins.Num()))
-	{
-		double TotalSum = 0;
-		double TotalObservations = 0;
-
-		for (int BinIdx = 0, LastBinIdx = Bins.Num() - 1; BinIdx < LastBinIdx; ++BinIdx)
-		{
-			FBin& Bin = Bins[BinIdx];
-			UE_LOG(LogPerfCounters, Log, TEXT("Bin %4.0f - %4.0f: %5d observation(s) which sum up to %f"), Bin.MinValue, Bin.UpperBound, Bin.Count, Bin.Sum);
-
-			TotalObservations += Bin.Count;
-			TotalSum += Bin.Sum;
-		}
-
-		FBin& LastBin = Bins.Last();
-		UE_LOG(LogPerfCounters, Log,     TEXT("Bin %4.0f +     : %5d observation(s) which sum up to %f"), LastBin.MinValue, LastBin.Count, LastBin.Sum);
-		TotalObservations += LastBin.Count;
-		TotalSum += LastBin.Sum;
-
-		if (TotalObservations > 0)
-		{
-			UE_LOG(LogPerfCounters, Log, TEXT("Average value for observation: %f"), TotalSum / TotalObservations);
-		}
-	}
-}
+const FName IPerfCounters::Histograms::ZeroLoadFrameTime(TEXT("ZeroLoadFrameTime"));
\ No newline at end of file
diff --git a/Engine/Source/Runtime/PerfCounters/Public/PerfCountersModule.h b/Engine/Source/Runtime/PerfCounters/Public/PerfCountersModule.h
--- a/Engine/Source/Runtime/PerfCounters/Public/PerfCountersModule.h
+++ b/Engine/Source/Runtime/PerfCounters/Public/PerfCountersModule.h
@@ -28,71 +28,6 @@ DECLARE_DELEGATE_OneParam(FProduceJsonCounterValue, const FPrettyJsonWriter& /*
  */
 DECLARE_DELEGATE_RetVal_TwoParams(bool, FPerfCounterExecCommandCallback, const FString& /*ExecCmd*/, FOutputDevice& /*Output*/);
 
-/** Fairly generic histogram for values that have natural lower bound and possibly no upper bound, e.g. frame time */
-struct PERFCOUNTERS_API FHistogram
-{
-	/** Inits histogram with linear, equally sized bins */
-	void InitLinear(double MinTime, double MaxTime, double BinSize);
-
-	/** Inits histogram to mimic our existing hitch buckets */
-	void InitHitchTracking();
-
-	/** Resets measurements, without resetting the configured bins. */
-	void Reset();
-
-	/** Adds an observed measurement. */
-	void AddMeasurement(double Value);
-
-	/** Prints histogram contents to the log. */
-	void DumpToLog(const FString& HistogramName);
-
-	/** Populates array commonly used in analytics events, adding two pairs per bin (count and sum). */
-	void DumpToAnalytics(const FString& ParamNamePrefix, TArray<TPair<FString, double>>& OutParamArray);
-
-protected:
-
-	/** Bin */
-	struct FBin
-	{
-		/** MinValue to be stored in the bin, inclusive. */
-		double				MinValue;
-
-		/** First value NOT to be stored in the bin. */
-		double				UpperBound;
-
-		/** Sum of all values that were put into this bin. */
-		double				Sum;
-
-		/** How many elements are in this bin. */
-		int32				Count;
-
-		FBin()
-		{
-		}
-
-		/** Constructor for any bin */
-		FBin(double MinInclusive, double MaxExclusive)
-			: MinValue(MinInclusive)
-			, UpperBound(MaxExclusive)
-			, Sum(0)
-			, Count(0)
-		{
-		}
-
-		/** Constructor for the last bin. */
-		FBin(double MinInclusive)
-			: MinValue(MinInclusive)
-			, UpperBound(FLT_MAX)
-			, Sum(0)
-			, Count(0)
-		{
-		}
-	};
-
-	/** Bins themselves, should be continous in terms of [MinValue; UpperBound) and sorted ascending by MinValue. Last bin's UpperBound doesn't matter */
-	TArray<FBin>			Bins;
-};
-
 /**
  * A programming interface for setting/updating performance counters
  */
@@ -124,6 +59,8 @@ public:
 		static const FName FrameTime;
 		/** Frame time histogram for shorter intervals. */
 		static const FName FrameTimePeriodic;
+		/** Frame time histogram (without sleep) for the duration of the match. */
+		static const FName FrameTimeWithoutSleep;
 		/** ServerReplicateActors time histogram for the duration of the match. */
 		static const FName ServerReplicateActorsTime;
 		/** Sleep time histogram for the duration of the match. */
diff --git a/Engine/Source/Runtime/Portal/LauncherCheck/LauncherCheck.Build.cs b/Engine/Source/Runtime/Portal/LauncherCheck/LauncherCheck.Build.cs
--- a/Engine/Source/Runtime/Portal/LauncherCheck/LauncherCheck.Build.cs
+++ b/Engine/Source/Runtime/Portal/LauncherCheck/LauncherCheck.Build.cs
@@ -10,13 +10,15 @@ public class LauncherCheck : ModuleRules
 			new string[]
 			{
 				"Core",
-                 "HTTP"
+                 "HTTP",
 			}
 		);
 
+        // Need to make this an option as it pulls in a developer module
         if (UEBuildConfiguration.bUseLauncherChecks)
         {
+            Definitions.Add("WITH_LAUNCHERCHECK=1");
             PublicDependencyModuleNames.Add("DesktopPlatform");
         }
-	}
+    }
 }
diff --git a/Engine/Source/Runtime/Portal/LauncherCheck/Private/LauncherCheckModule.cpp b/Engine/Source/Runtime/Portal/LauncherCheck/Private/LauncherCheckModule.cpp
--- a/Engine/Source/Runtime/Portal/LauncherCheck/Private/LauncherCheckModule.cpp
+++ b/Engine/Source/Runtime/Portal/LauncherCheck/Private/LauncherCheckModule.cpp
@@ -1,15 +1,13 @@
 // Copyright 1998-2016 Epic Games, Inc. All Rights Reserved.
 
 #include "LauncherCheckPrivatePCH.h"
-#include "GenericPlatformHttp.h"
 
 #if WITH_LAUNCHERCHECK
 
+#include "GenericPlatformHttp.h"
 #include "IDesktopPlatform.h"
 #include "DesktopPlatformModule.h"
 
-#endif
-
 /**
  * Log categories for LauncherCheck module
  */
@@ -23,6 +21,15 @@ class FLauncherCheckModule
 {
 public:
 
+	/*
+	* Check to see if this module should perform any checks or not
+	* @return true, if it should
+	*/
+	bool IsEnabled() const
+	{
+		return FParse::Param(FCommandLine::Get(), TEXT("NoEpicPortal")) == false;
+	}
+
 	// ILauncherCheckModule interface
 
 	virtual bool WasRanFromLauncher() const override
@@ -34,7 +41,6 @@ public:
 
 	virtual bool RunLauncher(ELauncherAction Action) const override
 	{
-#if WITH_LAUNCHERCHECK
 		IDesktopPlatform* DesktopPlatform = FDesktopPlatformModule::Get();
 		if (DesktopPlatform != nullptr)
 		{
@@ -53,7 +59,6 @@ public:
 			};
 			return DesktopPlatform->OpenLauncher(LauncherOptions);
 		}
-#endif // WITH_LAUNCHERCHECK
 		return false;
 	}
 
@@ -69,21 +74,6 @@ public:
 	{
 	}
 
-	/*
-	 * Check to see if this module should perform any checks or not
-	 * @return true, if it should
-	 */
-	bool IsEnabled() const
-	{
-#if WITH_LAUNCHERCHECK
-		// Check for the presence of a specific param that's passed from the Launcher to the game
-		// when we want to disabled the module
-		return !FParse::Param(FCommandLine::Get(), TEXT("NoEpicPortal"));
-#else
-		return false;
-#endif
-	}
-
 private:
 
 	/**
@@ -118,4 +108,22 @@ private:
 };
 
 
+#else
+
+class FLauncherCheckModule
+	: public ILauncherCheckModule
+{
+public:
+
+	virtual bool WasRanFromLauncher() const override { return true; }
+
+	virtual bool RunLauncher(ELauncherAction Action) const override { return false; }
+
+};
+
+
+#endif // WITH_LAUNCHERCHECK
+
 IMPLEMENT_MODULE(FLauncherCheckModule, LauncherCheck );
+
+
diff --git a/Engine/Source/Runtime/RHI/Public/RHIDefinitions.h b/Engine/Source/Runtime/RHI/Public/RHIDefinitions.h
--- a/Engine/Source/Runtime/RHI/Public/RHIDefinitions.h
+++ b/Engine/Source/Runtime/RHI/Public/RHIDefinitions.h
@@ -656,6 +656,24 @@ inline bool IsVulkanMobilePlatform(const EShaderPlatform Platform)
 	return Platform == SP_VULKAN_PCES3_1 || Platform == SP_VULKAN_ES3_1_ANDROID;
 }
 
+inline bool IsD3DPlatform(const EShaderPlatform Platform, bool bIncludeXboxOne)
+{
+	switch (Platform)
+	{
+	case SP_PCD3D_SM5:
+	case SP_PCD3D_SM4:
+	case SP_PCD3D_ES3_1:
+	case SP_PCD3D_ES2:
+		return true;
+	case SP_XBOXONE:
+		return bIncludeXboxOne;
+	default:
+		break;
+	}
+
+	return false;
+}
+
 inline ERHIFeatureLevel::Type GetMaxSupportedFeatureLevel(EShaderPlatform InShaderPlatform)
 {
 	switch (InShaderPlatform)
diff --git a/Engine/Source/Runtime/Renderer/Private/PostProcess/PostProcessing.cpp b/Engine/Source/Runtime/Renderer/Private/PostProcess/PostProcessing.cpp
--- a/Engine/Source/Runtime/Renderer/Private/PostProcess/PostProcessing.cpp
+++ b/Engine/Source/Runtime/Renderer/Private/PostProcess/PostProcessing.cpp
@@ -134,13 +134,24 @@ static TAutoConsoleVariable<float> CVarBloomCross(
 	TEXT(">0 for a cross look (X and Y)"),
 	ECVF_RenderThreadSafe);
 
-static TAutoConsoleVariable<int32> CVarTonemapperScreenPercentage(
-	TEXT("r.Tonemapper.ScreenPercentage"),
+static TAutoConsoleVariable<int32> CVarTonemapperMergeMode(
+	TEXT("r.Tonemapper.MergeWithUpscale.Mode"),
 	0,
-	TEXT("Experimental ScreenPercentage upscale integrated into tonemapper pass (if certain conditions apply e.g. no FXAA)\n")
+	TEXT("ScreenPercentage upscale integrated into tonemapper pass (if certain conditions apply, e.g., no FXAA)\n")
 	TEXT(" if enabled both features are done in one pass (faster, affects post process passes after the tonemapper including material post process e.g. sharpen)\n")
-	TEXT(" 0: off, the features run in separate passes (default)\n")
-	TEXT(" 1: enabled if nothing prevents it which is safe but it might do the two pass method"),
+	TEXT("  0: off, the features run in separate passes (default)\n")
+	TEXT("  1: always enabled, try to merge the passes unless something makes it impossible\n")
+	TEXT("  2: merge when the ratio of areas is above the r.Tonemapper.MergeWithUpscale.Threshold and it is otherwise possible"),
+	ECVF_Scalability | ECVF_RenderThreadSafe);
+
+static TAutoConsoleVariable<float> CVarTonemapperMergeThreshold(
+	TEXT("r.Tonemapper.MergeWithUpscale.Threshold"),
+	0.49f,
+	TEXT("If r.Tonemapper.MergeWithUpscale.Mode is 2, the ratio of the area before upscale/downscale to the area afterwards\n")
+	TEXT("is compared to this threshold when deciding whether or not to merge the passes.  The reasoning is that if the ratio\n")
+	TEXT("is too low, running the tonemapper on the higher number of pixels is more expensive than doing two passes\n")
+	TEXT("\n")
+	TEXT("Defauls to 0.49 (e.g., if r.ScreenPercentage is 70 or higher, try to merge)"),
 	ECVF_Scalability | ECVF_RenderThreadSafe);
 
 static TAutoConsoleVariable<int32> CVarMotionBlurScatter(
@@ -1837,11 +1848,23 @@ void FPostProcessing::Process(FRHICommandListImmediate& RHICmdList, const FViewI
 			// Check if we can save the Upscale pass and do it in the Tonemapper to save performance
 			if(Tonemapper && !PaniniConfig.IsEnabled() && !Tonemapper->bDoGammaOnly)
 			{
-				int32 TonemapperScreenPercentage = CVarTonemapperScreenPercentage.GetValueOnRenderThread();
-
-				if(TonemapperScreenPercentage != 0)
+				if (Context.FinalOutput.GetPass() == Tonemapper)
 				{
-					if(Context.FinalOutput.GetPass() == Tonemapper)
+					const int32 TonemapperMergeMode = CVarTonemapperMergeMode.GetValueOnRenderThread();
+					bool bCombineTonemapperAndUpsample = false;
+
+					if (TonemapperMergeMode == 1)
+					{
+						bCombineTonemapperAndUpsample = true;
+					}
+					else if (TonemapperMergeMode == 2)
+					{
+						const float TonemapperMergeThreshold = CVarTonemapperMergeThreshold.GetValueOnRenderThread();
+						const float AreaRatio = View.ViewRect.Area() / (float)View.UnscaledViewRect.Area();
+						bCombineTonemapperAndUpsample = AreaRatio > TonemapperMergeThreshold;
+					}
+
+					if (bCombineTonemapperAndUpsample)
 					{
 						Tonemapper->bDoScreenPercentageInTonemapper = true;
 						// the following pass is no longer needed
diff --git a/Engine/Source/Runtime/Renderer/Private/PrimitiveDistanceAccuracyRendering.cpp b/Engine/Source/Runtime/Renderer/Private/PrimitiveDistanceAccuracyRendering.cpp
--- a/Engine/Source/Runtime/Renderer/Private/PrimitiveDistanceAccuracyRendering.cpp
+++ b/Engine/Source/Runtime/Renderer/Private/PrimitiveDistanceAccuracyRendering.cpp
@@ -68,7 +68,7 @@ void FPrimitiveDistanceAccuracyPS::SetMesh(
 		}
 
 		const float OneOverDistanceMultiplier = 1.f / FMath::Max<float>(SMALL_NUMBER, ComponentExtraScale);
-		CPULogDistance =  FMath::Log2(OneOverDistanceMultiplier * FMath::Sqrt(FMath::Max<float>(1.f, DistSqMinusRadiusSq)));
+		CPULogDistance =  FMath::Max<float>(0.f, FMath::Log2(OneOverDistanceMultiplier * FMath::Sqrt(FMath::Max<float>(1.f, DistSqMinusRadiusSq))));
 	}
 #endif
 	// Because the streamer use FMath::FloorToFloat, here we need to use -1 to have a useful result.
diff --git a/Engine/Source/Runtime/ShaderCore/Private/Shader.cpp b/Engine/Source/Runtime/ShaderCore/Private/Shader.cpp
--- a/Engine/Source/Runtime/ShaderCore/Private/Shader.cpp
+++ b/Engine/Source/Runtime/ShaderCore/Private/Shader.cpp
@@ -1638,13 +1638,18 @@ void ShaderMapAppendKeyString(EShaderPlatform Platform, FString& KeyString)
 	}
 	
 	{
+		// Always default to fast math unless specified
 		static const auto CVar = IConsoleManager::Get().FindConsoleVariable(TEXT("r.Shaders.FastMath"));
-		KeyString += (CVar && CVar->GetInt() != 0) ? TEXT("") : TEXT("_NoFastMath");
+		KeyString += (CVar && CVar->GetInt() == 0) ? TEXT("_NoFastMath") : TEXT("");
 	}
-	
+
+	if (IsD3DPlatform(Platform, false))
 	{
-		static const auto CVar = IConsoleManager::Get().FindConsoleVariable(TEXT("r.Shaders.AvoidFlowControl"));
-		KeyString += (CVar && CVar->GetInt() != 0) ? TEXT("_Unroll") : TEXT("_Flow");
+		static const auto CVar = IConsoleManager::Get().FindConsoleVariable(TEXT("r.D3D.RemoveUnusedInterpolators"));
+		if (CVar && CVar->GetInt() != 0)
+		{
+			KeyString += TEXT("_UnInt");
+		}
 	}
 
 	if (Platform == SP_PS4)
diff --git a/Engine/Source/Runtime/ShaderCore/Public/ShaderCore.h b/Engine/Source/Runtime/ShaderCore/Public/ShaderCore.h
--- a/Engine/Source/Runtime/ShaderCore/Public/ShaderCore.h
+++ b/Engine/Source/Runtime/ShaderCore/Public/ShaderCore.h
@@ -114,7 +114,9 @@ enum ECompilerFlags
 	CFLAG_KeepDebugInfo,
 	CFLAG_NoFastMath,
 	// Compile ES2 with ES3.1 features
-	CFLAG_FeatureLevelES31
+	CFLAG_FeatureLevelES31,
+	// Force removing unused interpolators for platforms that can opt out
+	CFLAG_ForceRemoveUnusedInterpolators,
 };
 
 /**
diff --git a/Engine/Source/Runtime/Slate/Private/Widgets/Input/SEditableTextBox.cpp b/Engine/Source/Runtime/Slate/Private/Widgets/Input/SEditableTextBox.cpp
--- a/Engine/Source/Runtime/Slate/Private/Widgets/Input/SEditableTextBox.cpp
+++ b/Engine/Source/Runtime/Slate/Private/Widgets/Input/SEditableTextBox.cpp
@@ -198,6 +198,11 @@ FText SEditableTextBox::GetSelectedText() const
 	return EditableText->GetSelectedText();
 }
 
+bool SEditableTextBox::HasError() const
+{
+	return ErrorReporting.IsValid() && ErrorReporting->HasError();
+}
+
 const FSlateBrush* SEditableTextBox::GetBorderImage() const
 {
 	if ( EditableText->IsTextReadOnly() )
diff --git a/Engine/Source/Runtime/Slate/Private/Widgets/Notifications/SPopUpErrorText.cpp b/Engine/Source/Runtime/Slate/Private/Widgets/Notifications/SPopUpErrorText.cpp
--- a/Engine/Source/Runtime/Slate/Private/Widgets/Notifications/SPopUpErrorText.cpp
+++ b/Engine/Source/Runtime/Slate/Private/Widgets/Notifications/SPopUpErrorText.cpp
@@ -13,11 +13,13 @@ void SPopupErrorText::Construct( const FArguments& InArgs )
 		[
 			SAssignNew( HasErrorSymbol, SErrorText )
 			.ErrorText(NSLOCTEXT("UnrealEd", "Error", "!"))
+			.Font( InArgs._Font )
 		]
 		.MenuPlacement(MenuPlacement_BelowAnchor)
 		.MenuContent()
 		[
 			SAssignNew( ErrorText, SErrorText )
+			.Font( InArgs._Font )
 		]
 	);
 }
diff --git a/Engine/Source/Runtime/Slate/Private/Widgets/SInvalidationPanel.cpp b/Engine/Source/Runtime/Slate/Private/Widgets/SInvalidationPanel.cpp
--- a/Engine/Source/Runtime/Slate/Private/Widgets/SInvalidationPanel.cpp
+++ b/Engine/Source/Runtime/Slate/Private/Widgets/SInvalidationPanel.cpp
@@ -303,7 +303,7 @@ int32 SInvalidationPanel::OnPaint( const FPaintArgs& Args, const FGeometry& Allo
 		{
 			INC_DWORD_STAT_BY(STAT_SlateNumCachedElements, CachedWindowElements->GetDrawElements().Num());
 
-			RootCacheNode->RecordHittestGeometry(Args.GetGrid(), Args.GetLastHitTestIndex());
+			RootCacheNode->RecordHittestGeometry(Args.GetGrid(), Args.GetLastHitTestIndex(), LayerId);
 		}
 		else
 		{
diff --git a/Engine/Source/Runtime/Slate/Private/Widgets/Text/SlateEditableTextTypes.cpp b/Engine/Source/Runtime/Slate/Private/Widgets/Text/SlateEditableTextTypes.cpp
--- a/Engine/Source/Runtime/Slate/Private/Widgets/Text/SlateEditableTextTypes.cpp
+++ b/Engine/Source/Runtime/Slate/Private/Widgets/Text/SlateEditableTextTypes.cpp
@@ -128,7 +128,7 @@ int32 FCursorLineHighlighter::OnPaint(const FPaintArgs& Args, const FTextLayout:
 	FSlateDrawElement::MakeBox(
 		OutDrawElements,
 		LayerId,
-		AllottedGeometry.ToPaintGeometry(TransformVector(InverseScale, FVector2D(CursorWidth * AllottedGeometry.Scale, Size.Y)), FSlateLayoutTransform(TransformPoint(InverseScale, Location + OptionalWidth))),
+		AllottedGeometry.ToPaintGeometry(TransformVector(InverseScale, FVector2D(FMath::Max(CursorWidth * AllottedGeometry.Scale, 1.0f), Size.Y)), FSlateLayoutTransform(TransformPoint(InverseScale, Location + OptionalWidth))),
 		CursorBrush.Get(),
 		MyClippingRect,
 		bParentEnabled ? ESlateDrawEffect::None : ESlateDrawEffect::DisabledEffect,
diff --git a/Engine/Source/Runtime/Slate/Public/Widgets/Input/SEditableTextBox.h b/Engine/Source/Runtime/Slate/Public/Widgets/Input/SEditableTextBox.h
--- a/Engine/Source/Runtime/Slate/Public/Widgets/Input/SEditableTextBox.h
+++ b/Engine/Source/Runtime/Slate/Public/Widgets/Input/SEditableTextBox.h
@@ -249,6 +249,8 @@ public:
 	/** Get the currently selected text */
 	FText GetSelectedText() const;
 
+	bool HasError() const;
+
 	// SWidget overrides
 	virtual bool SupportsKeyboardFocus() const override;
 	virtual bool HasKeyboardFocus() const override;
diff --git a/Engine/Source/Runtime/Slate/Public/Widgets/Notifications/SPopUpErrorText.h b/Engine/Source/Runtime/Slate/Public/Widgets/Notifications/SPopUpErrorText.h
--- a/Engine/Source/Runtime/Slate/Public/Widgets/Notifications/SPopUpErrorText.h
+++ b/Engine/Source/Runtime/Slate/Public/Widgets/Notifications/SPopUpErrorText.h
@@ -9,9 +9,11 @@ class SLATE_API SPopupErrorText : public SComboButton, public IErrorReportingWid
 public:
 	SLATE_BEGIN_ARGS(SPopupErrorText)
 		: _ShowInNewWindow( false )
+		, _Font()
 	{}
 		/** The popup appears in a new window instead of in the same window that this widget is in */
 		SLATE_ARGUMENT( bool, ShowInNewWindow )
+ 		SLATE_ATTRIBUTE( FSlateFontInfo, Font )
 	SLATE_END_ARGS()
 
 	virtual void Construct( const FArguments& InArgs );
diff --git a/Engine/Source/Runtime/SlateCore/Private/Input/HittestGrid.cpp b/Engine/Source/Runtime/SlateCore/Private/Input/HittestGrid.cpp
--- a/Engine/Source/Runtime/SlateCore/Private/Input/HittestGrid.cpp
+++ b/Engine/Source/Runtime/SlateCore/Private/Input/HittestGrid.cpp
@@ -93,12 +93,13 @@ struct FHittestGrid::FGridTestingParams
 
 struct FHittestGrid::FCachedWidget
 {
-	FCachedWidget(int32 InParentIndex, const FArrangedWidget& InWidget, const FSlateRect& InClippingRect)
+	FCachedWidget(int32 InParentIndex, const FArrangedWidget& InWidget, const FSlateRect& InClippingRect, int32 InLayerId)
 	: WidgetPtr(InWidget.Widget)
 	, CachedGeometry(InWidget.Geometry)
 	, ClippingRect(InClippingRect)
 	, Children()
 	, ParentIndex(InParentIndex)
+	, LayerId(InLayerId)
 	{}
 
 	void AddChild(const int32 ChildIndex)
@@ -114,6 +115,8 @@ struct FHittestGrid::FCachedWidget
 	FSlateRect ClippingRect;
 	TArray<int32, TInlineAllocator<16> > Children;
 	int32 ParentIndex;
+	/** This is needed to be able to pick the best of the widgets within the virtual cursor's radius. */
+	int32 LayerId;
 };
 
 
@@ -205,15 +208,32 @@ TArray<FWidgetAndPointer> FHittestGrid::GetBubblePath(FVector2D DesktopSpaceCoor
 				return A.DistToTopWidgetSq < B.DistToTopWidgetSq;
 			});
 
+			const FWidgetPathAndDist* BestCandidateByDistance = nullptr;
 			for (const FWidgetPathAndDist& TestPath : PathsAndDistances)
 			{
 				if (ContainsInteractableWidget(TestPath.BubblePath))
 				{
-					return TestPath.BubblePath;
+					BestCandidateByDistance = &TestPath;
+					break;
+				}
+			}
+
+			if (BestCandidateByDistance != nullptr)
+			{
+				const FGeometry& LeafmostGeometry = BestCandidateByDistance->BubblePath.Last().Geometry;
+				const int32 LeafmostLayerId = BestCandidateByDistance->LayerId;
+				if (DirectBubblePathInfo.LayerId > LeafmostLayerId)
+				{
+					return DirectBubblePathInfo.BubblePath;
+				}
+				else
+				{
+					return BestCandidateByDistance->BubblePath;
 				}
 			}
 		}
 
+
 		return DirectBubblePathInfo.BubblePath;
 	}
 	else
@@ -249,7 +269,7 @@ void FHittestGrid::ClearGridForNewFrame( const FSlateRect& HittestArea )
 	}
 }
 
-int32 FHittestGrid::InsertWidget( const int32 ParentHittestIndex, const EVisibility& Visibility, const FArrangedWidget& Widget, const FVector2D InWindowOffset, const FSlateRect& InClippingRect )
+int32 FHittestGrid::InsertWidget(const int32 ParentHittestIndex, const EVisibility& Visibility, const FArrangedWidget& Widget, const FVector2D InWindowOffset, const FSlateRect& InClippingRect, int32 LayerId)
 {
 	if ( ensureMsgf( ParentHittestIndex < WidgetsCachedThisFrame->Num(), TEXT("Widget '%s' being drawn before its parent."), *Widget.ToString() ) )
 	{
@@ -260,7 +280,7 @@ int32 FHittestGrid::InsertWidget( const int32 ParentHittestIndex, const EVisibil
 
 
 		// Remember this widget, its geometry, and its place in the logical hierarchy.
-		const int32 WidgetIndex = WidgetsCachedThisFrame->Add( FCachedWidget( ParentHittestIndex, WindowAdjustedWidget, WindowAdjustedRect ) );
+		const int32 WidgetIndex = WidgetsCachedThisFrame->Add(FCachedWidget(ParentHittestIndex, WindowAdjustedWidget, WindowAdjustedRect, LayerId));
 		check( WidgetIndex < WidgetsCachedThisFrame->Num() ); 
 		if (ParentHittestIndex != INDEX_NONE)
 		{
@@ -627,13 +647,13 @@ FHittestGrid::FWidgetPathAndDist FHittestGrid::GetWidgetPathAndDist(const FGridT
 			TArray<FWidgetAndPointer> LogicalBubblePath = GetBubblePathFromHitIndex(HitIndex.WidgetIndex, bIgnoreEnabledStatus);
 			const TArray<FWidgetAndPointer> BubblePathExtension = PhysicallyHitWidget.CustomPath.Pin()->GetBubblePathAndVirtualCursors(PhysicallyHitWidget.CachedGeometry, DesktopSpaceCoordinate, bIgnoreEnabledStatus);
 			LogicalBubblePath.Append(BubblePathExtension);
-			return FWidgetPathAndDist(LogicalBubblePath, 0.0);
+			return FWidgetPathAndDist(LogicalBubblePath, 0.0, PhysicallyHitWidget.LayerId);
 		}
 		else
 		{
 			//get the path from the hit index, check if anything came back
 			const TArray<FWidgetAndPointer> BubblePath = GetBubblePathFromHitIndex(HitIndex.WidgetIndex, bIgnoreEnabledStatus);
-			return FWidgetPathAndDist(BubblePath, BubblePath.Num() > 0 ? HitIndex.DistanceSqToWidget : -1.0f);
+			return FWidgetPathAndDist(BubblePath, BubblePath.Num() > 0 ? HitIndex.DistanceSqToWidget : -1.0f, PhysicallyHitWidget.LayerId);
 		}
 	}
 
diff --git a/Engine/Source/Runtime/SlateCore/Private/Layout/WidgetCaching.cpp b/Engine/Source/Runtime/SlateCore/Private/Layout/WidgetCaching.cpp
--- a/Engine/Source/Runtime/SlateCore/Private/Layout/WidgetCaching.cpp
+++ b/Engine/Source/Runtime/SlateCore/Private/Layout/WidgetCaching.cpp
@@ -24,19 +24,19 @@ void FCachedWidgetNode::Initialize(const FPaintArgs& Args, TSharedRef<SWidget> I
 	Children.Reset();
 }
 
-void FCachedWidgetNode::RecordHittestGeometry(FHittestGrid& Grid, int32 LastHittestIndex)
+void FCachedWidgetNode::RecordHittestGeometry(FHittestGrid& Grid, int32 LastHittestIndex, int32 LayerId)
 {
 	if ( RecordedVisibility.AreChildrenHitTestVisible() )
 	{
 		TSharedPtr<SWidget> SafeWidget = Widget.Pin();
 		if ( SafeWidget.IsValid() )
 		{
-			LastRecordedHittestIndex = Grid.InsertWidget(LastHittestIndex, RecordedVisibility, FArrangedWidget(SafeWidget.ToSharedRef(), Geometry), WindowOffset, ClippingRect);
+			LastRecordedHittestIndex = Grid.InsertWidget(LastHittestIndex, RecordedVisibility, FArrangedWidget(SafeWidget.ToSharedRef(), Geometry), WindowOffset, ClippingRect, LayerId);
 
 			const int32 ChildCount = Children.Num();
 			for ( int32 i = 0; i < ChildCount; i++ )
 			{
-				Children[i]->RecordHittestGeometry(Grid, LastRecordedHittestIndex);
+				Children[i]->RecordHittestGeometry(Grid, LastRecordedHittestIndex, LayerId);
 			}
 		}
 	}
diff --git a/Engine/Source/Runtime/SlateCore/Private/Types/PaintArgs.cpp b/Engine/Source/Runtime/SlateCore/Private/Types/PaintArgs.cpp
--- a/Engine/Source/Runtime/SlateCore/Private/Types/PaintArgs.cpp
+++ b/Engine/Source/Runtime/SlateCore/Private/Types/PaintArgs.cpp
@@ -30,7 +30,7 @@ FPaintArgs FPaintArgs::EnableCaching(ILayoutCache* InLayoutCache, FCachedWidgetN
 	return UpdatedArgs;
 }
 
-FPaintArgs FPaintArgs::RecordHittestGeometry(const SWidget* Widget, const FGeometry& WidgetGeometry, const FSlateRect& InClippingRect) const
+FPaintArgs FPaintArgs::RecordHittestGeometry(const SWidget* Widget, const FGeometry& WidgetGeometry, int32 LayerId, const FSlateRect& InClippingRect) const
 {
 	FPaintArgs UpdatedArgs(*this);
 
@@ -54,7 +54,7 @@ FPaintArgs FPaintArgs::RecordHittestGeometry(const SWidget* Widget, const FGeome
 
 		// When rendering volatile widgets, their parent widgets who have been cached 
 		const EVisibility RecordedVisibility = Widget->GetVisibility();
-		const int32 RecordedHittestIndex = Grid.InsertWidget(RealLastHitTestIndex, RecordedVisibility, FArrangedWidget(const_cast<SWidget*>( Widget )->AsShared(), WidgetGeometry), WindowOffset, InClippingRect);
+		const int32 RecordedHittestIndex = Grid.InsertWidget(RealLastHitTestIndex, RecordedVisibility, FArrangedWidget(const_cast<SWidget*>(Widget)->AsShared(), WidgetGeometry), WindowOffset, InClippingRect, LayerId);
 		UpdatedArgs.LastHittestIndex = RecordedHittestIndex;
 		UpdatedArgs.LastRecordedVisibility = RecordedVisibility;
 	}
diff --git a/Engine/Source/Runtime/SlateCore/Private/Widgets/SWidget.cpp b/Engine/Source/Runtime/SlateCore/Private/Widgets/SWidget.cpp
--- a/Engine/Source/Runtime/SlateCore/Private/Widgets/SWidget.cpp
+++ b/Engine/Source/Runtime/SlateCore/Private/Widgets/SWidget.cpp
@@ -654,7 +654,7 @@ int32 SWidget::Paint(const FPaintArgs& Args, const FGeometry& AllottedGeometry,
 	}
 
 	// Record hit test geometry, but only if we're not caching.
-	const FPaintArgs UpdatedArgs = Args.RecordHittestGeometry(this, AllottedGeometry, MyClippingRect);
+	const FPaintArgs UpdatedArgs = Args.RecordHittestGeometry(this, AllottedGeometry, LayerId, MyClippingRect);
 
 	// Paint the geometry of this widget.
 	int32 NewLayerID = OnPaint(UpdatedArgs, AllottedGeometry, MyClippingRect, OutDrawElements, LayerId, InWidgetStyle, bParentEnabled);
diff --git a/Engine/Source/Runtime/SlateCore/Public/Input/HittestGrid.h b/Engine/Source/Runtime/SlateCore/Public/Input/HittestGrid.h
--- a/Engine/Source/Runtime/SlateCore/Public/Input/HittestGrid.h
+++ b/Engine/Source/Runtime/SlateCore/Public/Input/HittestGrid.h
@@ -39,7 +39,7 @@ public:
 	void ClearGridForNewFrame(const FSlateRect& HittestArea);
 
 	/** Add Widget into the hittest data structure so that we can later make queries about it. */
-	int32 InsertWidget( const int32 ParentHittestIndex, const EVisibility& Visibility, const FArrangedWidget& Widget, const FVector2D InWindowOffset, const FSlateRect& InClippingRect );
+	int32 InsertWidget(const int32 ParentHittestIndex, const EVisibility& Visibility, const FArrangedWidget& Widget, const FVector2D InWindowOffset, const FSlateRect& InClippingRect, int32 LayerId);
 
 	void InsertCustomHitTestPath( TSharedRef<ICustomHitTestPath> CustomHitTestPath, int32 WidgetIndex );
 
@@ -81,17 +81,19 @@ private:
 	bool IsValidCellCoord(const FIntPoint& CellCoord) const;
 	bool IsValidCellCoord(const int32 XCoord, const int32 YCoord) const;
 
-	/** Bubble path and distance to leafmost widget from point of picking. */
+	/** Bubble path, distance to leafmost widget from point of picking, and LayerId of leafmost widget. */
 	struct FWidgetPathAndDist
 	{
 		/** Ctor */
 		FWidgetPathAndDist( float InDistanceSq = -1.0f )
 		: DistToTopWidgetSq(InDistanceSq)
+		, LayerId(0)
 		{}
 
-		FWidgetPathAndDist(const TArray<FWidgetAndPointer>& InPath, float InDistanceSq)
+		FWidgetPathAndDist(const TArray<FWidgetAndPointer>& InPath, float InDistanceSq, int32 InLayerId)
 		: BubblePath(InPath)
 		, DistToTopWidgetSq(InDistanceSq)
+		, LayerId(InLayerId)
 		{}
 
 		void Clear()
@@ -107,6 +109,7 @@ private:
 
 		TArray<FWidgetAndPointer> BubblePath;
 		float DistToTopWidgetSq;
+		int32 LayerId;
 	};
 
 	FWidgetPathAndDist GetWidgetPathAndDist(const FGridTestingParams& Params, const bool bIgnoreEnabledStatus) const;
diff --git a/Engine/Source/Runtime/SlateCore/Public/Layout/WidgetCaching.h b/Engine/Source/Runtime/SlateCore/Public/Layout/WidgetCaching.h
--- a/Engine/Source/Runtime/SlateCore/Public/Layout/WidgetCaching.h
+++ b/Engine/Source/Runtime/SlateCore/Public/Layout/WidgetCaching.h
@@ -12,7 +12,7 @@ class SLATECORE_API FCachedWidgetNode
 public:
 	void Initialize(const FPaintArgs& Args, TSharedRef<SWidget> InWidget, const FGeometry& InGeometry, const FSlateRect& InClippingRect);
 
-	void RecordHittestGeometry(FHittestGrid& Grid, int32 LastHittestIndex);
+	void RecordHittestGeometry(FHittestGrid& Grid, int32 LastHittestIndex, int32 LayerId);
 
 public:
 
diff --git a/Engine/Source/Runtime/SlateCore/Public/Types/PaintArgs.h b/Engine/Source/Runtime/SlateCore/Public/Types/PaintArgs.h
--- a/Engine/Source/Runtime/SlateCore/Public/Types/PaintArgs.h
+++ b/Engine/Source/Runtime/SlateCore/Public/Types/PaintArgs.h
@@ -43,7 +43,7 @@ public:
 	}
 
 	FPaintArgs EnableCaching(ILayoutCache* InLayoutCache, FCachedWidgetNode* InParentCacheNode, bool bEnableCaching, bool bEnableVolatility) const;
-	FPaintArgs RecordHittestGeometry(const SWidget* Widget, const FGeometry& WidgetGeometry, const FSlateRect& InClippingRect) const;
+	FPaintArgs RecordHittestGeometry(const SWidget* Widget, const FGeometry& WidgetGeometry, int32 LayerId, const FSlateRect& InClippingRect) const;
 	FPaintArgs InsertCustomHitTestPath( TSharedRef<ICustomHitTestPath> CustomHitTestPath, int32 HitTestIndex ) const;
 
 	FHittestGrid& GetGrid() const { return Grid; }
diff --git a/Engine/Source/Runtime/SlateRHIRenderer/Private/Slate3DRenderer.cpp b/Engine/Source/Runtime/SlateRHIRenderer/Private/Slate3DRenderer.cpp
--- a/Engine/Source/Runtime/SlateRHIRenderer/Private/Slate3DRenderer.cpp
+++ b/Engine/Source/Runtime/SlateRHIRenderer/Private/Slate3DRenderer.cpp
@@ -92,7 +92,7 @@ struct TKeepAliveCommand : public FRHICommand < TKeepAliveCommand<TKeepAliveType
 	void Execute(FRHICommandListBase& CmdList) {}
 };
 
-void FSlate3DRenderer::DrawWindowToTarget_RenderThread( FRHICommandListImmediate& InRHICmdList, FTextureRenderTarget2DResource* RenderTargetResource, FSlateDrawBuffer& WindowDrawBuffer )
+void FSlate3DRenderer::DrawWindowToTarget_RenderThread( FRHICommandListImmediate& InRHICmdList, FTextureRenderTarget2DResource* RenderTargetResource, FSlateDrawBuffer& WindowDrawBuffer, bool bInClearTarget)
 {
 	SCOPED_DRAW_EVENT( InRHICmdList, SlateRenderToTarget );
 
@@ -112,7 +112,7 @@ void FSlate3DRenderer::DrawWindowToTarget_RenderThread( FRHICommandListImmediate
 	FRHIRenderTargetView ColorRTV(RTResource);
 	ColorRTV.LoadAction = ERenderTargetLoadAction::EClear;
 	FRHISetRenderTargetsInfo Info(1, &ColorRTV, FTextureRHIParamRef());
-	Info.bClearColor = true;
+	Info.bClearColor = bInClearTarget;
 
 	InRHICmdList.TransitionResource(EResourceTransitionAccess::EWritable, RTResource);
 	InRHICmdList.SetRenderTargetsAndClear(Info);
diff --git a/Engine/Source/Runtime/SlateRHIRenderer/Private/Slate3DRenderer.h b/Engine/Source/Runtime/SlateRHIRenderer/Private/Slate3DRenderer.h
--- a/Engine/Source/Runtime/SlateRHIRenderer/Private/Slate3DRenderer.h
+++ b/Engine/Source/Runtime/SlateRHIRenderer/Private/Slate3DRenderer.h
@@ -14,7 +14,7 @@ public:
 
 	virtual FSlateDrawBuffer& GetDrawBuffer() override;
 	virtual void DrawWindow_GameThread(FSlateDrawBuffer& DrawBuffer) override;
-	virtual void DrawWindowToTarget_RenderThread(FRHICommandListImmediate& RHICmdList, FTextureRenderTarget2DResource* RenderTargetResource, FSlateDrawBuffer& InDrawBuffer) override;
+	virtual void DrawWindowToTarget_RenderThread(FRHICommandListImmediate& RHICmdList, FTextureRenderTarget2DResource* RenderTargetResource, FSlateDrawBuffer& InDrawBuffer, bool bInClearTarget) override;
 
 private:
 
diff --git a/Engine/Source/Runtime/SlateRHIRenderer/Public/Interfaces/ISlate3DRenderer.h b/Engine/Source/Runtime/SlateRHIRenderer/Public/Interfaces/ISlate3DRenderer.h
--- a/Engine/Source/Runtime/SlateRHIRenderer/Public/Interfaces/ISlate3DRenderer.h
+++ b/Engine/Source/Runtime/SlateRHIRenderer/Public/Interfaces/ISlate3DRenderer.h
@@ -25,7 +25,7 @@ public:
 	 * @param RenderTarget The render target to render the contents of the draw buffer to
 	 * @param InDrawBuffer The draw buffer containing the batched elements to render
 	 */
-	virtual void DrawWindowToTarget_RenderThread(FRHICommandListImmediate& RHICmdList, FTextureRenderTarget2DResource* RenderTargetResource, FSlateDrawBuffer& InDrawBuffer) = 0;
+	virtual void DrawWindowToTarget_RenderThread(FRHICommandListImmediate& RHICmdList, FTextureRenderTarget2DResource* RenderTargetResource, FSlateDrawBuffer& InDrawBuffer, bool bInClearTarget = true) = 0;
 };
 
 typedef TSharedPtr<ISlate3DRenderer, ESPMode::ThreadSafe> ISlate3DRendererPtr;
\ No newline at end of file
diff --git a/Engine/Source/Runtime/UMG/Private/Components/EditableTextBox.cpp b/Engine/Source/Runtime/UMG/Private/Components/EditableTextBox.cpp
--- a/Engine/Source/Runtime/UMG/Private/Components/EditableTextBox.cpp
+++ b/Engine/Source/Runtime/UMG/Private/Components/EditableTextBox.cpp
@@ -120,6 +120,16 @@ void UEditableTextBox::ClearError()
 	}
 }
 
+bool UEditableTextBox::HasError() const
+{
+	if ( MyEditableTextBlock.IsValid() )
+	{
+		return MyEditableTextBlock->HasError();
+	}
+
+	return false;
+}
+
 void UEditableTextBox::HandleOnTextChanged(const FText& InText)
 {
 	Text = InText;
diff --git a/Engine/Source/Runtime/UMG/Private/Slate/SlateVectorArtData.cpp b/Engine/Source/Runtime/UMG/Private/Slate/SlateVectorArtData.cpp
--- a/Engine/Source/Runtime/UMG/Private/Slate/SlateVectorArtData.cpp
+++ b/Engine/Source/Runtime/UMG/Private/Slate/SlateVectorArtData.cpp
@@ -148,9 +148,9 @@ void USlateVectorArtData::EnsureValidData()
 #endif
 }
 
-void USlateVectorArtData::PreSave()
+void USlateVectorArtData::PreSave(const class ITargetPlatform* TargetPlatform)
 {
-	Super::PreSave();
+	Super::PreSave(TargetPlatform);
 	EnsureValidData();
 }
 
diff --git a/Engine/Source/Runtime/UMG/Private/Slate/WidgetRenderer.cpp b/Engine/Source/Runtime/UMG/Private/Slate/WidgetRenderer.cpp
--- a/Engine/Source/Runtime/UMG/Private/Slate/WidgetRenderer.cpp
+++ b/Engine/Source/Runtime/UMG/Private/Slate/WidgetRenderer.cpp
@@ -64,9 +64,10 @@ void SVirtualWindow::OnArrangeChildren(const FGeometry& AllottedGeometry, FArran
 	}
 }
 
-FWidgetRenderer::FWidgetRenderer(bool bUseGammaCorrection)
+FWidgetRenderer::FWidgetRenderer(bool bUseGammaCorrection, bool bInClearTarget)
 	: bPrepassNeeded(true)
 	, bUseGammaSpace(bUseGammaCorrection)
+	, bClearTarget(bInClearTarget)
 {
 #if !UE_SERVER
 	if (!IsRunningDedicatedServer())
@@ -205,19 +206,21 @@ void FWidgetRenderer::DrawWindow(
 			FSlateDrawBuffer* DrawBuffer;
 			FTextureRenderTarget2DResource* RenderTargetResource;
 			TSharedPtr<ISlate3DRenderer, ESPMode::ThreadSafe> Renderer;
-		};
-		FRenderThreadContext Context =
+			bool bClearTarget;
+		}
+		Context =		
 		{
 			&DrawBuffer,
 			static_cast<FTextureRenderTarget2DResource*>(RenderTarget->GameThread_GetRenderTargetResource()),
-			Renderer
+			Renderer,
+			bClearTarget
 		};
 
 		// Enqueue a command to unlock the draw buffer after all windows have been drawn
 		ENQUEUE_UNIQUE_RENDER_COMMAND_ONEPARAMETER(FWidgetRenderer_DrawWindow,
 			FRenderThreadContext, InContext, Context,
 			{
-				InContext.Renderer->DrawWindowToTarget_RenderThread(RHICmdList, InContext.RenderTargetResource, *InContext.DrawBuffer);
+				InContext.Renderer->DrawWindowToTarget_RenderThread(RHICmdList, InContext.RenderTargetResource, *InContext.DrawBuffer, InContext.bClearTarget);
 			});
 	}
 #endif // !UE_SERVER
diff --git a/Engine/Source/Runtime/UMG/Private/UserWidget.cpp b/Engine/Source/Runtime/UMG/Private/UserWidget.cpp
--- a/Engine/Source/Runtime/UMG/Private/UserWidget.cpp
+++ b/Engine/Source/Runtime/UMG/Private/UserWidget.cpp
@@ -818,9 +818,9 @@ void UUserWidget::RemoveObsoleteBindings(const TArray<FName>& NamedSlots)
 	}
 }
 
-void UUserWidget::PreSave()
+void UUserWidget::PreSave(const class ITargetPlatform* TargetPlatform)
 {
-	Super::PreSave();
+	Super::PreSave(TargetPlatform);
 
 	// Remove bindings that are no longer contained in the class.
 	if ( UWidgetBlueprintGeneratedClass* BGClass = Cast<UWidgetBlueprintGeneratedClass>(GetClass()) )
diff --git a/Engine/Source/Runtime/UMG/Private/WidgetTree.cpp b/Engine/Source/Runtime/UMG/Private/WidgetTree.cpp
--- a/Engine/Source/Runtime/UMG/Private/WidgetTree.cpp
+++ b/Engine/Source/Runtime/UMG/Private/WidgetTree.cpp
@@ -96,13 +96,13 @@ void UWidgetTree::GetChildWidgets(UWidget* Parent, TArray<UWidget*>& Widgets)
 	ForWidgetAndChildren(Parent, [&] (UWidget* Widget) { Widgets.Add(Widget); });
 }
 
-void UWidgetTree::PreSave()
+void UWidgetTree::PreSave(const class ITargetPlatform* TargetPlatform)
 {
 	AllWidgets.Empty();
 
 	GetAllWidgets(AllWidgets);
 
-	Super::PreSave();
+	Super::PreSave( TargetPlatform);
 }
 
 void UWidgetTree::PostLoad()
diff --git a/Engine/Source/Runtime/UMG/Public/Components/EditableTextBox.h b/Engine/Source/Runtime/UMG/Public/Components/EditableTextBox.h
--- a/Engine/Source/Runtime/UMG/Public/Components/EditableTextBox.h
+++ b/Engine/Source/Runtime/UMG/Public/Components/EditableTextBox.h
@@ -141,6 +141,9 @@ public:
 	UFUNCTION(BlueprintCallable, Category="Widget")
 	void ClearError();
 
+	UFUNCTION(BlueprintCallable, Category = "Widget")
+	bool HasError() const;
+
 	//~ Begin UWidget Interface
 	virtual void SynchronizeProperties() override;
 	//~ End UWidget Interface
diff --git a/Engine/Source/Runtime/UMG/Public/Slate/WidgetRenderer.h b/Engine/Source/Runtime/UMG/Public/Slate/WidgetRenderer.h
--- a/Engine/Source/Runtime/UMG/Public/Slate/WidgetRenderer.h
+++ b/Engine/Source/Runtime/UMG/Public/Slate/WidgetRenderer.h
@@ -33,11 +33,12 @@ private:
 class UMG_API FWidgetRenderer
 {
 public:
-	FWidgetRenderer(bool bUseGammaCorrection = false);
+	FWidgetRenderer(bool bUseGammaCorrection = false, bool bInClearTarget = true);
 	~FWidgetRenderer();
 
 	bool GetIsPrepassNeeded() const { return bPrepassNeeded; }
 	void SetIsPrepassNeeded(bool bInPrepassNeeded) { bPrepassNeeded = bInPrepassNeeded; }
+	void SetShouldClearTarget(bool bShouldClear) { bClearTarget = bShouldClear; }
 
 	ISlate3DRenderer* GetSlateRenderer();
 
@@ -86,4 +87,6 @@ private:
 	bool bPrepassNeeded;
 	/** Is gamma space needed? */
 	bool bUseGammaSpace;
+	/** Should we clear the render target before rendering. */
+	bool bClearTarget;
 };
diff --git a/Engine/Source/Runtime/WebBrowser/Private/WebBrowserSingleton.cpp b/Engine/Source/Runtime/WebBrowser/Private/WebBrowserSingleton.cpp
--- a/Engine/Source/Runtime/WebBrowser/Private/WebBrowserSingleton.cpp
+++ b/Engine/Source/Runtime/WebBrowser/Private/WebBrowserSingleton.cpp
@@ -191,7 +191,8 @@ FWebBrowserSingleton::FWebBrowserSingleton()
 	// Specify path to sub process exe
 	FString SubProcessPath(FPaths::Combine(*FPaths::EngineDir(), CEF3_SUBPROCES_EXE));
 	SubProcessPath = FPaths::ConvertRelativePathToFull(SubProcessPath);
-	if (!FPaths::FileExists(SubProcessPath))
+
+	if (!IPlatformFile::GetPlatformPhysical().FileExists(*SubProcessPath))
 	{
 		UE_LOG(LogWebBrowser, Error, TEXT("UnrealCEFSubProcess.exe not found, check that this program has been built and is placed in: %s."), *SubProcessPath);
 	}
diff --git a/Engine/Source/Runtime/Windows/D3D11RHI/Private/D3D11Commands.cpp b/Engine/Source/Runtime/Windows/D3D11RHI/Private/D3D11Commands.cpp
--- a/Engine/Source/Runtime/Windows/D3D11RHI/Private/D3D11Commands.cpp
+++ b/Engine/Source/Runtime/Windows/D3D11RHI/Private/D3D11Commands.cpp
@@ -59,6 +59,14 @@ static FAutoConsoleVariableRef CVarDX11TransitionChecks(
 	ECVF_Default
 	);
 
+int32 GUnbindResourcesBetweenDrawsInDX11 = 0;
+static FAutoConsoleVariableRef CVarUnbindResourcesBetweenDrawsInDX11(
+	TEXT("r.UnbindResourcesBetweenDrawsInDX11"),
+	GUnbindResourcesBetweenDrawsInDX11,
+	TEXT("Unbind resources between material changes in DX11."),
+	ECVF_Default
+	);
+
 void FD3D11BaseShaderResource::SetDirty(bool bInDirty, uint32 CurrentFrame)
 {
 	bDirty = bInDirty;
@@ -233,6 +241,11 @@ void FD3D11DynamicRHI::RHISetBoundShaderState( FBoundShaderStateRHIParamRef Boun
 			BoundUniformBuffers[Frequency][BindIndex].SafeRelease();
 		}
 	}
+
+	if (GUnbindResourcesBetweenDrawsInDX11)
+	{
+		ClearAllShaderResources();
+	}
 }
 
 void FD3D11DynamicRHI::RHISetShaderTexture(FVertexShaderRHIParamRef VertexShaderRHI,uint32 TextureIndex,FTextureRHIParamRef NewTextureRHI)
