Make some spammy perf counters logging verbose.

diff --git a/Engine/Source/Runtime/PerfCounters/Private/PerfCounters.cpp b/Engine/Source/Runtime/PerfCounters/Private/PerfCounters.cpp
--- a/Engine/Source/Runtime/PerfCounters/Private/PerfCounters.cpp
+++ b/Engine/Source/Runtime/PerfCounters/Private/PerfCounters.cpp
@@ -228,7 +228,7 @@ void FPerfCounters::TickSocket(float DeltaTime)
 		if (bLogConnections)
 		{
 			IncomingConnection->GetPeerAddress(*ScratchIPAddr);
-			UE_LOG(LogPerfCounters, Log, TEXT("New connection from %s"), *ScratchIPAddr->ToString(true));
+			UE_LOG(LogPerfCounters, Verbose, TEXT("New connection from %s"), *ScratchIPAddr->ToString(true));
 		}
 
 		// make sure this is non-blocking
@@ -288,7 +288,7 @@ void FPerfCounters::TickSocket(float DeltaTime)
 
 				double EndTime = FPlatformTime::Seconds();
 				ExistingSocket->GetPeerAddress(*ScratchIPAddr);
-				UE_LOG(LogPerfCounters, Log, TEXT("Request for %s processed in %0.2f s"), *ScratchIPAddr->ToString(true), EndTime - StartTime);
+				UE_LOG(LogPerfCounters, Verbose, TEXT("Request for %s processed in %0.2f s"), *ScratchIPAddr->ToString(true), EndTime - StartTime);
 			}
 			else
 			{
@@ -317,7 +317,7 @@ void FPerfCounters::TickSocket(float DeltaTime)
 			if (bLogConnectionClosure)
 			{
 				ClosingSocket->GetPeerAddress(*ScratchIPAddr);
-				UE_LOG(LogPerfCounters, Log, TEXT("Closed connection to %s."), *ScratchIPAddr->ToString(true));
+				UE_LOG(LogPerfCounters, Verbose, TEXT("Closed connection to %s."), *ScratchIPAddr->ToString(true));
 			}
 
 			// close the socket (whether we processed or not)
