[PerfCounters] PerfCounters HTTP interface leverages HTTPServer module
+ Internal HTTP bind/listen/read/write functionality has been replaced with HTTPServer module usage
+ All endpoints/functionality are identical.
+ Minor structural change: PerfCountersModule now stores a TSharedPtr to FPerfCounters rather than a raw pointer.  This change has been made to facilitate usage of AsShared() within FPerfCounters to facilitate safe and non-intrusive captures.

diff --git a/Engine/Source/Runtime/PerfCounters/PerfCounters.Build.cs b/Engine/Source/Runtime/PerfCounters/PerfCounters.Build.cs
--- a/Engine/Source/Runtime/PerfCounters/PerfCounters.Build.cs
+++ b/Engine/Source/Runtime/PerfCounters/PerfCounters.Build.cs
@@ -11,9 +11,8 @@ public class PerfCounters : ModuleRules
 				"Core",
 				"CoreUObject",
 				"Json",
-				"Sockets",
-				"HTTP"
-			}
+                "HTTPServer"
+            }
         );
 	}
 }
diff --git a/Engine/Source/Runtime/PerfCounters/Private/PerfCounters.cpp b/Engine/Source/Runtime/PerfCounters/Private/PerfCounters.cpp
--- a/Engine/Source/Runtime/PerfCounters/Private/PerfCounters.cpp
+++ b/Engine/Source/Runtime/PerfCounters/Private/PerfCounters.cpp
@@ -10,11 +10,13 @@
 #include "Misc/ConfigCacheIni.h"
 #include "Serialization/JsonWriter.h"
 #include "ZeroLoad.h"
-#include "SocketSubsystem.h"
-#include "IPAddress.h"
-#include "Sockets.h"
-#include "PlatformHttp.h"
-#include "Interfaces/IHttpResponse.h"
+
+#include "HttpServerModule.h"
+#include "IHttpRouter.h"
+#include "HttpserverRequest.h"
+#include "HttpServerResponse.h"
+#include "HttpRequestHandler.h"
+#include "HttpServerConstants.h"
 
 #define JSON_ARRAY_NAME					TEXT("PerfCounters")
 #define JSON_PERFCOUNTER_NAME			TEXT("Name")
@@ -23,10 +25,8 @@
 #define PERF_COUNTER_CONNECTION_TIMEOUT 5.0f
 
 FPerfCounters::FPerfCounters(const FString& InUniqueInstanceId)
-	: SocketSubsystem(nullptr)
-	, UniqueInstanceId(InUniqueInstanceId)
+	: UniqueInstanceId(InUniqueInstanceId)
 	, InternalCountersUpdateInterval(60)
-	, Socket(nullptr)
 	, ZeroLoadThread(nullptr)
 	, ZeroLoadRunnable(nullptr)
 {
@@ -34,14 +34,10 @@ FPerfCounters::FPerfCounters(const FString& InUniqueInstanceId)
 
 FPerfCounters::~FPerfCounters()
 {
-	if (Socket)
+	if (HttpRouter)
 	{
-		ISocketSubsystem* SocketSystem = ISocketSubsystem::Get(PLATFORM_SOCKETSUBSYSTEM);
-		if (SocketSystem)
-		{
-			SocketSystem->DestroySocket(Socket);
-		}
-		Socket = nullptr;
+		HttpRouter->UnbindRoute(StatsRouteHandle);
+		HttpRouter->UnbindRoute(ExecRouteHandle);
 	}
 }
 
@@ -63,47 +59,43 @@ bool FPerfCounters::Initialize()
 		return true;
 	}
 
-	// get the socket subsystem
-	SocketSubsystem = ISocketSubsystem::Get(PLATFORM_SOCKETSUBSYSTEM);
-	if (SocketSubsystem == nullptr)
+	// Get an IHttpRouter on the command-line designated port
+	HttpRouter = FHttpServerModule::Get().GetHttpRouter(StatsPort);
+	if (!HttpRouter)
 	{
-		UE_LOG(LogPerfCounters, Error, TEXT("FPerfCounters unable to get socket subsystem"));
+		UE_LOG(LogPerfCounters, Error, 
+			TEXT("FPerfCounters unable to bind to specified statsPort [%d]"), StatsPort);
 		return false;
 	}
 
-	ScratchIPAddr = SocketSubsystem->CreateInternetAddr();
-
-	// make our listen socket
-	Socket = SocketSubsystem->CreateSocket(NAME_Stream, TEXT("FPerfCounters"));
-	if (Socket == nullptr)
+	// Register a handler for /stats
+	TWeakPtr<FPerfCounters> WeakThisPtr(AsShared());
+    StatsRouteHandle = HttpRouter->BindRoute(FHttpPath("/stats"), EHttpServerRequestVerbs::VERB_GET,
+		[WeakThisPtr](const FHttpServerRequest& Request, const FHttpResultCallback& OnComplete)
 	{
-		UE_LOG(LogPerfCounters, Error, TEXT("FPerfCounters unable to allocate stream socket"));
-		return false;
-	}
-
-	// make us non blocking
-	Socket->SetNonBlocking(true);
-
-	// create a localhost binding for the requested port
-	TSharedRef<FInternetAddr> LocalhostAddr = SocketSubsystem->CreateInternetAddr(0x7f000001 /* 127.0.0.1 */, StatsPort);
-	if (!Socket->Bind(*LocalhostAddr))
+		auto SharedThis = WeakThisPtr.Pin();
+		if (!SharedThis.IsValid()) { return false; }
+		return SharedThis->ProcessStatsRequest(Request, OnComplete);
+	});
+	if(!StatsRouteHandle.IsValid())
 	{
-		UE_LOG(LogPerfCounters, Error, TEXT("FPerfCounters unable to bind to %s"), *LocalhostAddr->ToString(true));
+		UE_LOG(LogPerfCounters, Error,
+			TEXT("FPerfCounters unable bind route: /stats"));
 		return false;
 	}
-	StatsPort = Socket->GetPortNo();
-
-	// log the port
-	UE_LOG(LogPerfCounters, Display, TEXT("FPerfCounters listening on port %d"), StatsPort);
-
-	// for now, jack this up so we can send in one go
-	int32 NewSize;
-	Socket->SetSendBufferSize(512 * 1024, NewSize); // best effort 512k buffer to avoid not being able to send in one go
 
-	// listen on the port
-	if (!Socket->Listen(16))
+	// Register a handler for /exec
+	ExecRouteHandle = HttpRouter->BindRoute(FHttpPath("/exec"), EHttpServerRequestVerbs::VERB_GET,
+		[WeakThisPtr] (const FHttpServerRequest& Request, const FHttpResultCallback& OnComplete)
 	{
-		UE_LOG(LogPerfCounters, Error, TEXT("FPerfCounters unable to listen on socket"));
+		auto SharedThis = WeakThisPtr.Pin();
+		if (!SharedThis.IsValid()) { return false; }
+		return SharedThis->ProcessExecRequest(Request, OnComplete);
+	});
+	if(!ExecRouteHandle.IsValid())
+	{
+		UE_LOG(LogPerfCounters, Error,
+			TEXT("FPerfCounters unable bind route: /exec"));
 		return false;
 	}
 
@@ -162,34 +154,10 @@ void FPerfCounters::ResetStatsForNextPeriod()
 	}
 };
 
-static bool SendAsUtf8(FSocket* Conn, const FString& Message)
-{
-	FTCHARToUTF8 ConvertToUtf8(*Message);
-	int32 BytesSent = 0;
-	const bool bSendSuccess = Conn->Send(reinterpret_cast<const uint8*>(ConvertToUtf8.Get()), ConvertToUtf8.Length(), BytesSent);
-	const bool bSendSizeSuccess = (BytesSent == ConvertToUtf8.Length());
-
-	if (!bSendSuccess)
-	{
-		UE_LOG(LogPerfCounters, Warning, TEXT("Failed to send buffer size: %d"), ConvertToUtf8.Length());
-	}
-	else if (!bSendSizeSuccess)
-	{
-		UE_LOG(LogPerfCounters, Warning, TEXT("Failed to send entire buffer size: %d sent: %d"), ConvertToUtf8.Length(), BytesSent);
-	}
-
-	return bSendSuccess && bSendSizeSuccess;
-}
-
 bool FPerfCounters::Tick(float DeltaTime)
 {
     QUICK_SCOPE_CYCLE_COUNTER(STAT_FPerfCounters_Tick);
 
-	if (LIKELY(Socket != nullptr))
-	{
-		TickSocket(DeltaTime);
-	}
-
 	if (LIKELY(ZeroLoadThread != nullptr))
 	{
 		TickZeroLoad(DeltaTime);
@@ -215,117 +183,6 @@ void FPerfCounters::TickZeroLoad(float DeltaTime)
 	}
 }
 
-void FPerfCounters::TickSocket(float DeltaTime)
-{
-	checkf(Socket != nullptr, TEXT("FPerfCounters::TickSocket() called without a valid socket!"));
-
-	// accept any connections
-	static const FString PerfCounterRequest = TEXT("FPerfCounters Request");
-	FSocket* IncomingConnection = Socket->Accept(PerfCounterRequest);
-	if (IncomingConnection)
-	{
-		const bool bLogConnections = true;
-		if (bLogConnections)
-		{
-			IncomingConnection->GetPeerAddress(*ScratchIPAddr);
-			UE_LOG(LogPerfCounters, Verbose, TEXT("New connection from %s"), *ScratchIPAddr->ToString(true));
-		}
-
-		// make sure this is non-blocking
-		IncomingConnection->SetNonBlocking(true);
-
-		new (Connections) FPerfConnection(IncomingConnection);
-	}
-	else
-	{
-		ESocketErrors ErrorCode = SocketSubsystem->GetLastErrorCode();
-		if (ErrorCode != SE_EWOULDBLOCK &&
-			ErrorCode != SE_NO_ERROR)
-		{
-			const TCHAR* ErrorStr = SocketSubsystem->GetSocketError();
-			UE_LOG(LogPerfCounters, Warning, TEXT("Error accepting connection [%d] %s"), (int32)ErrorCode, ErrorStr);
-		}
-	}
-
-	TArray<FPerfConnection> ConnectionsToClose;
-	for (FPerfConnection& Connection : Connections)
-	{
-		FSocket* ExistingSocket = Connection.Connection;
-		if (ExistingSocket && ExistingSocket->Wait(ESocketWaitConditions::WaitForRead, FTimespan::Zero()))
-		{
-			// read any data that's ready
-			// NOTE: this is not a full HTTP implementation, just enough to be usable by curl
-			uint8 Buffer[2 * 1024] = { 0 };
-			int32 DataLen = 0;
-			if (ExistingSocket->Recv(Buffer, sizeof(Buffer) - 1, DataLen, ESocketReceiveFlags::None))
-			{
-				double StartTime = FPlatformTime::Seconds();
-
-				FResponse Response;
-				if (ProcessRequest(Buffer, DataLen, Response))
-				{
-					if (!EHttpResponseCodes::IsOk(Response.Code))
-					{
-						UE_LOG(LogPerfCounters, Warning, TEXT("Sending error response: [%d] %s"), Response.Code, *Response.Body);
-					}
-
-					if (SendAsUtf8(ExistingSocket, Response.Header))
-					{
-						if (!SendAsUtf8(ExistingSocket, Response.Body))
-						{
-							UE_LOG(LogPerfCounters, Warning, TEXT("Unable to send full HTTP response body size: %d"), Response.Body.Len());
-						}
-					}
-					else
-					{
-						UE_LOG(LogPerfCounters, Warning, TEXT("Unable to send HTTP response header: %s"), *Response.Header);
-					}
-				}
-				else
-				{
-					UE_LOG(LogPerfCounters, Warning, TEXT("Failed to process request"));
-				}
-
-				double EndTime = FPlatformTime::Seconds();
-				ExistingSocket->GetPeerAddress(*ScratchIPAddr);
-				UE_LOG(LogPerfCounters, Verbose, TEXT("Request for %s processed in %0.2f s"), *ScratchIPAddr->ToString(true), EndTime - StartTime);
-			}
-			else
-			{
-				UE_LOG(LogPerfCounters, Warning, TEXT("Unable to immediately receive request header"));
-			}
-
-			ConnectionsToClose.Add(Connection);
-		}
-		else if (Connection.ElapsedTime > PERF_COUNTER_CONNECTION_TIMEOUT)
-		{
-			UE_LOG(LogPerfCounters, Warning, TEXT("Closing connection due to timeout %d"), Connection.ElapsedTime);
-			ConnectionsToClose.Add(Connection);
-		}
-
-		Connection.ElapsedTime += DeltaTime;
-	}
-
-	for (FPerfConnection& Connection : ConnectionsToClose)
-	{
-		Connections.RemoveSingleSwap(Connection);
-
-		FSocket* ClosingSocket = Connection.Connection;
-		if (ClosingSocket)
-		{
-			const bool bLogConnectionClosure = true;
-			if (bLogConnectionClosure)
-			{
-				ClosingSocket->GetPeerAddress(*ScratchIPAddr);
-				UE_LOG(LogPerfCounters, Verbose, TEXT("Closed connection to %s."), *ScratchIPAddr->ToString(true));
-			}
-
-			// close the socket (whether we processed or not)
-			ClosingSocket->Close();
-			SocketSubsystem->DestroySocket(ClosingSocket);
-		}
-	}
-}
 
 void FPerfCounters::TickSystemCounters(float DeltaTime)
 {
@@ -373,77 +230,42 @@ bool FPerfCounters::Exec(UWorld* InWorld, const TCHAR* Cmd, FOutputDevice& Ar)
 	return false;
 }
 
-bool FPerfCounters::ProcessRequest(uint8* Buffer, int32 BufferLen, FResponse& Response)
+bool FPerfCounters::ProcessStatsRequest(const FHttpServerRequest& Request, const FHttpResultCallback& OnComplete)
 {
-	bool bSuccess = false;
+	auto ResponseBody = GetAllCountersAsJson();
+	auto Response = FHttpServerResponse::Create(ResponseBody, TEXT("application/json"));
+	OnComplete(MoveTemp(Response));
+	return true;
+}
 
-	// scan the buffer for a line
-	FUTF8ToTCHAR WideBuffer(reinterpret_cast<const ANSICHAR*>(Buffer));
-	const TCHAR* BufferEnd = FCString::Strstr(WideBuffer.Get(), TEXT("\r\n"));
-	if (BufferEnd != nullptr)
+bool FPerfCounters::ProcessExecRequest(const FHttpServerRequest& Request, const FHttpResultCallback& OnComplete)
+{
+	FStringOutputDevice StringOutDevice;
+	bool bExecCommandSuccess = false;
+	const FString* ExecCmd = Request.QueryParams.Find(TEXT("C"));
+	if (ExecCmd)
 	{
-		// crack into pieces
-		FString MainLine(BufferEnd - WideBuffer.Get(), WideBuffer.Get());
-		TArray<FString> Tokens;
-		MainLine.ParseIntoArrayWS(Tokens);
-		if (Tokens.Num() >= 2)
-		{
-			FString ContentType(TEXT("application/json"));
-			Response.Code = EHttpResponseCodes::Ok;
+		StringOutDevice.SetAutoEmitLineTerminator(true);
 
-			// handle the request
-			if (Tokens[0] != TEXT("GET"))
-			{
-				Response.Body = FString::Printf(TEXT("{ \"error\": \"Method %s not allowed\" }"), *Tokens[0]);
-				Response.Code = EHttpResponseCodes::BadMethod;
-			}
-			else if (Tokens[1].StartsWith(TEXT("/stats")))
-			{
-				Response.Body = GetAllCountersAsJson();
-			}
-			else if (Tokens[1].StartsWith(TEXT("/exec?c=")))
-			{
-				FString ExecCmd = Tokens[1].Mid(8);
-				FString ExecCmdDecoded = FPlatformHttp::UrlDecode(ExecCmd);
-
-				FStringOutputDevice StringOutDevice;
-				StringOutDevice.SetAutoEmitLineTerminator(true);
-
-				bool bResult = false;
-				if (ExecCmdCallback.IsBound())
-				{
-					bResult = ExecCmdCallback.Execute(ExecCmdDecoded, StringOutDevice);
-					Response.Body = StringOutDevice;
-					ContentType = TEXT("text/text");
-				}
-				else
-				{
-					Response.Body = FString::Printf(TEXT("{ \"error\": \"exec handler not found\" }"));
-				}
-
-				Response.Code = bResult ? EHttpResponseCodes::Ok : EHttpResponseCodes::NotFound;
-			}
-			else
-			{
-				Response.Body = FString::Printf(TEXT("{ \"error\": \"%s not found\" }"), *Tokens[1]);
-				Response.Code = EHttpResponseCodes::NotFound;
-			}
-
-			// send the response headers
-			Response.Header = FString::Printf(TEXT("HTTP/1.0 %d\r\nContent-Length: %d\r\nContent-Type: %s\r\n\r\n"), Response.Code, Response.Body.Len(), *ContentType);
-			bSuccess = true;
-		}
-		else
+		if (ExecCmdCallback.IsBound())
 		{
-			UE_LOG(LogPerfCounters, Warning, TEXT("Unable to parse HTTP request header: %s"), *MainLine);
+			bExecCommandSuccess = ExecCmdCallback.Execute(*ExecCmd, StringOutDevice);
 		}
 	}
+
+	if (bExecCommandSuccess)
+	{
+		auto Response = FHttpServerResponse::Create(StringOutDevice, TEXT("text/text"));
+		OnComplete(MoveTemp(Response));
+	}
 	else
 	{
-		UE_LOG(LogPerfCounters, Warning, TEXT("ProcessRequest: request incomplete"));
+		auto Response = FHttpServerResponse::Error(EHttpServerResponseCodes::NotSupported, 
+			TEXT("exec handler not found"));
+		OnComplete(MoveTemp(Response));
 	}
 
-	return bSuccess;
+	return true;
 }
 
 double FPerfCounters::GetNumber(const FString& Name, double DefaultValue)
diff --git a/Engine/Source/Runtime/PerfCounters/Private/PerfCounters.h b/Engine/Source/Runtime/PerfCounters/Private/PerfCounters.h
--- a/Engine/Source/Runtime/PerfCounters/Private/PerfCounters.h
+++ b/Engine/Source/Runtime/PerfCounters/Private/PerfCounters.h
@@ -8,8 +8,12 @@
 #include "Containers/Ticker.h"
 #include "ProfilingDebugging/Histogram.h"
 
-class ISocketSubsystem;
-class FSocket;
+#include "HttpResultCallback.h"
+#include "HttpPath.h"
+#include "HttpRouteHandle.h"
+
+class IHttpRouter;
+struct FHttpServerRequest;
 class FZeroLoad;
 class FInternetAddr;
 
@@ -19,6 +23,7 @@ class FPerfCounters
 	: public FTickerObjectBase
 	, public FSelfRegisteringExec
 	, public IPerfCounters
+	, public TSharedFromThis<FPerfCounters>
 {
 public:
 
@@ -57,71 +62,25 @@ public:
 private:
 	
 	void TickZeroLoad(float DeltaTime);
-	void TickSocket(float DeltaTime);
 	void TickSystemCounters(float DeltaTime);
 
 	/**
-	 * Simple response structure for returning output to requestor
-	 */
-	struct FResponse
-	{
-		/** http header */
-		FString Header;
-		/** http body */
-		FString Body;
-		/** http response code */
-		int32 Code;
-
-		FResponse() :
-			Code(0)
-		{}
-	};
-	
-	/**
-	 * Simple connection structure for keeping track of incoming/active connections
+	 * Processes a /stats request
+	 *
+	 * @param Request The incoming request
+	 * @param OnComplete The invokable response result callback
+	 * @return true if this request was handled herein, false otherwise
 	 */
-	struct FPerfConnection
-	{
-		/** accepted external socket */
-		FSocket* Connection;
-		/** time connection has existed */
-		float ElapsedTime;
-
-		FPerfConnection() :
-			Connection(nullptr),
-			ElapsedTime(0.0f)
-		{}
-
-		FPerfConnection(FSocket* InConnection) :
-			Connection(InConnection),
-			ElapsedTime(0.0f)
-		{}
-
-		bool operator==(const FPerfConnection& Other) const
-		{
-			return Connection == Other.Connection;
-		}
-	};
-
-	/** Reference to the socket subsystem in use */
-	ISocketSubsystem* SocketSubsystem;
-
-	/** Scratch IP address to save on memory allocations */
-	TSharedPtr<FInternetAddr> ScratchIPAddr;
-
-	/** all active connections */
-	TArray<FPerfConnection> Connections;
+	bool ProcessStatsRequest(const FHttpServerRequest& Request, const FHttpResultCallback& OnComplete);
 
 	/**
-	 * Process the incoming request from an active socket
+	 * Process an /exec request
 	 *
-	 * @param Buffer data sent from the requestor
-	 * @param BufferLen size of the data sent
-	 * @param Response [out] response to be given to the requestor
-	 *
-	 * @return true if the response is valid, false otherwise
+	 * @param Request The incoming request
+	 * @param OnComplete The invokable response result callback
+	 * @return true if this request was handled herein, false otherwise
 	 */
-	bool ProcessRequest(uint8* Buffer, int32 BufferLen, FResponse& Response);
+	bool ProcessExecRequest(const FHttpServerRequest& Request, const FHttpResultCallback& OnComplete);
 
 	/** Unique name of this instance */
 	FString UniqueInstanceId;
@@ -138,9 +97,6 @@ private:
 	/** Bound callback for script command execution */
 	FPerfCounterExecCommandCallback ExecCmdCallback;
 
-	/* Listen socket for outputting JSON on request */
-	FSocket* Socket;
-
 	/** Map of performance histograms. */
 	TPerformanceHistogramMap PerformanceHistogramMap;
 
@@ -149,4 +105,13 @@ private:
 
 	/** Zero-load thread runnable */
 	FRunnableThread* ZeroLoadRunnable;
+
+	/** Http request/response interface */
+	TSharedPtr<IHttpRouter> HttpRouter = nullptr;
+
+	/** Route handle for /stats binding */
+	FHttpRouteHandle StatsRouteHandle = nullptr;
+
+	/** Route handle for /exec binding */
+	FHttpRouteHandle ExecRouteHandle = nullptr;
 };
diff --git a/Engine/Source/Runtime/PerfCounters/Private/PerfCountersModule.cpp b/Engine/Source/Runtime/PerfCounters/Private/PerfCountersModule.cpp
--- a/Engine/Source/Runtime/PerfCounters/Private/PerfCountersModule.cpp
+++ b/Engine/Source/Runtime/PerfCounters/Private/PerfCountersModule.cpp
@@ -9,21 +9,17 @@ class FPerfCountersModule : public IPerfCountersModule
 private:
 
 	/** All created perf counter instances from this module */
-	FPerfCounters* PerfCountersSingleton;
+	TSharedPtr<FPerfCounters> PerfCountersSingleton = nullptr;
 
 public:
 
-	FPerfCountersModule() : 
-		PerfCountersSingleton(NULL)
-	{}
+	FPerfCountersModule() 
+	{
+	}
 
-	void ShutdownModule()
+	virtual void ShutdownModule() override
 	{
-		if (PerfCountersSingleton)
-		{
-			delete PerfCountersSingleton;
-			PerfCountersSingleton = nullptr;
-		}
+		PerfCountersSingleton.Reset();
 	}
 
 	virtual bool SupportsDynamicReloading() override
@@ -38,7 +34,7 @@ public:
 
 	IPerfCounters* GetPerformanceCounters() const
 	{
-		return PerfCountersSingleton;
+		return PerfCountersSingleton.Get();
 	}
 
 	IPerfCounters* CreatePerformanceCounters(const FString& UniqueInstanceId) override
@@ -46,7 +42,7 @@ public:
 		if (PerfCountersSingleton)
 		{
 			UE_LOG(LogPerfCounters, Display, TEXT("CreatePerformanceCounters: instance already exists, new instance not created."));
-			return PerfCountersSingleton;
+			return PerfCountersSingleton.Get();
 		}
 
 		FString InstanceUID = UniqueInstanceId;
@@ -54,17 +50,15 @@ public:
 		{
 			InstanceUID = FString::Printf(TEXT("perfcounters-of-pid-%d"), FPlatformProcess::GetCurrentProcessId());
 		}
-
-		FPerfCounters* PerfCounters = new FPerfCounters(InstanceUID);
-		if (!PerfCounters->Initialize())
+ 
+		PerfCountersSingleton = MakeShared<FPerfCounters>(InstanceUID);
+		if (!PerfCountersSingleton->Initialize())
 		{
 			UE_LOG(LogPerfCounters, Warning, TEXT("CreatePerformanceCounters: could not create perfcounters"));
-			delete PerfCounters;
 			return nullptr;
 		}
 
-		PerfCountersSingleton = PerfCounters;
-		return PerfCounters;
+		return PerfCountersSingleton.Get();
 	}
 };
 
