[PerfCounters] Fix #include directive casing to fix linux compile failure

diff --git a/Engine/Source/Runtime/PerfCounters/Private/PerfCounters.cpp b/Engine/Source/Runtime/PerfCounters/Private/PerfCounters.cpp
--- a/Engine/Source/Runtime/PerfCounters/Private/PerfCounters.cpp
+++ b/Engine/Source/Runtime/PerfCounters/Private/PerfCounters.cpp
@@ -13,7 +13,7 @@
 
 #include "HttpServerModule.h"
 #include "IHttpRouter.h"
-#include "HttpserverRequest.h"
+#include "HttpServerRequest.h"
 #include "HttpServerResponse.h"
 #include "HttpRequestHandler.h"
 #include "HttpServerConstants.h"
