[PerfCounters] Build Fix - Previously inferred includes are now explicitly #include'd

diff --git a/Engine/Source/Runtime/PerfCounters/Private/PerfCounters.cpp b/Engine/Source/Runtime/PerfCounters/Private/PerfCounters.cpp
--- a/Engine/Source/Runtime/PerfCounters/Private/PerfCounters.cpp
+++ b/Engine/Source/Runtime/PerfCounters/Private/PerfCounters.cpp
@@ -6,9 +6,12 @@
 #include "HAL/FileManager.h"
 #include "Misc/CommandLine.h"
 #include "Misc/Paths.h"
+#include "HAL/PlatformTime.h"
+#include "HAL/PlatformProcess.h"
 #include "HAL/RunnableThread.h"
 #include "Misc/ConfigCacheIni.h"
 #include "Serialization/JsonWriter.h"
+#include "Stats/Stats.h"
 #include "ZeroLoad.h"
 
 #include "HttpServerModule.h"
