PerfCounters
* Added -forkautostatsport cmdline option that automatically sets a unique port for each forked child. Ports used will start at  statsPort+1.
* Improved statsport /exec listener so it returns a different error message when the exec callback isn't registered, the cmd itself was not processed by the exec or the cmd is missing the 'c=' param.

Forking
* Exposed ChildProcessIndex in ForkHelper

diff --git a/Engine/Source/Runtime/Core/Private/Misc/Fork.cpp b/Engine/Source/Runtime/Core/Private/Misc/Fork.cpp
--- a/Engine/Source/Runtime/Core/Private/Misc/Fork.cpp
+++ b/Engine/Source/Runtime/Core/Private/Misc/Fork.cpp
@@ -25,6 +25,9 @@ namespace ForkInternal
 	/** Will be set to true on a child process right before we spin up the ForkableThreads at the end of the frame where we forked. */
 	static bool bIsForkedMultithreadInstance = false;
 
+	/** The unique index of a forked child process */
+	static uint16 ChildProcessIndex = 0;
+
 	/**
 	 * Are we doing a real fork and generating child processes on the fork process who received the -WaitAndFork commandline.
 	 * Note: this will be true on the child processes too.
@@ -73,9 +76,15 @@ bool FForkProcessHelper::IsForkedChildProcess()
 	return ForkInternal::bIsForkedChildProcess;
 }
 
-void FForkProcessHelper::SetIsForkedChildProcess()
+void FForkProcessHelper::SetIsForkedChildProcess(uint16 ChildIndex)
 {
 	ForkInternal::bIsForkedChildProcess = true;
+	ForkInternal::ChildProcessIndex = ChildIndex;
+}
+
+uint16 FForkProcessHelper::GetForkedChildProcessIndex()
+{
+	return ForkInternal::ChildProcessIndex;
 }
 
 bool FForkProcessHelper::IsForkRequested()
diff --git a/Engine/Source/Runtime/Core/Private/Unix/UnixPlatformProcess.cpp b/Engine/Source/Runtime/Core/Private/Unix/UnixPlatformProcess.cpp
--- a/Engine/Source/Runtime/Core/Private/Unix/UnixPlatformProcess.cpp
+++ b/Engine/Source/Runtime/Core/Private/Unix/UnixPlatformProcess.cpp
@@ -1461,7 +1461,10 @@ FGenericPlatformProcess::EWaitAndForkResult FUnixPlatformProcess::WaitAndFork()
 			}
 			else if (ChildPID == 0)
 			{
-				FForkProcessHelper::SetIsForkedChildProcess();
+				// Child
+				uint16 Cookie = (SignalValue >> 16) & 0xffff;
+				uint16 ChildIdx = SignalValue & 0xffff;
+				FForkProcessHelper::SetIsForkedChildProcess(ChildIdx);
 
 				if (FPlatformMemory::HasForkPageProtectorEnabled())
 				{
@@ -1469,10 +1472,6 @@ FGenericPlatformProcess::EWaitAndForkResult FUnixPlatformProcess::WaitAndFork()
 					UE::FForkPageProtector::Get().ProtectMemoryRegions();
 				}
 
-				// Child
-				uint16 Cookie = (SignalValue >> 16) & 0xffff;
-				uint16 ChildIdx = SignalValue & 0xffff;
-
 				// Close the log state we inherited from our parent
 				GLog->TearDown();
 
diff --git a/Engine/Source/Runtime/Core/Public/Misc/Fork.h b/Engine/Source/Runtime/Core/Public/Misc/Fork.h
--- a/Engine/Source/Runtime/Core/Public/Misc/Fork.h
+++ b/Engine/Source/Runtime/Core/Public/Misc/Fork.h
@@ -42,9 +42,14 @@ public:
 	static bool IsForkedChildProcess();
 
 	/**
-	 * Sets the forked child process flag 
+	 * Sets the forked child process flag and index given to this child process
 	 */
-	static void SetIsForkedChildProcess();
+	static void SetIsForkedChildProcess(uint16 ChildIndex);
+
+	/**
+	* Returns the unique index of this forked child process. Index 0 is for the master server
+	*/
+	static uint16 GetForkedChildProcessIndex();
 
 	/**
 	 * Event triggered when a fork occurred on the child process and its safe to create real threads
diff --git a/Engine/Source/Runtime/PerfCounters/Private/PerfCounters.cpp b/Engine/Source/Runtime/PerfCounters/Private/PerfCounters.cpp
--- a/Engine/Source/Runtime/PerfCounters/Private/PerfCounters.cpp
+++ b/Engine/Source/Runtime/PerfCounters/Private/PerfCounters.cpp
@@ -253,7 +253,20 @@ bool FPerfCounters::ProcessExecRequest(const FHttpServerRequest& Request, const
 		{
 			bExecCommandSuccess = ExecCmdCallback.Execute(*ExecCmd, StringOutDevice);
 		}
+		else
+		{
+			auto Response = FHttpServerResponse::Error(EHttpServerResponseCodes::NotSupported,
+				TEXT("exec handler not found"));
+			OnComplete(MoveTemp(Response));
+		}
 	}
+	else
+	{
+		auto Response = FHttpServerResponse::Error(EHttpServerResponseCodes::NotSupported,
+			TEXT("exec missing query command (c=MyCommand)"));
+		OnComplete(MoveTemp(Response));
+	}
+
 
 	if (bExecCommandSuccess)
 	{
@@ -262,8 +275,7 @@ bool FPerfCounters::ProcessExecRequest(const FHttpServerRequest& Request, const
 	}
 	else
 	{
-		auto Response = FHttpServerResponse::Error(EHttpServerResponseCodes::NotSupported, 
-			TEXT("exec handler not found"));
+		auto Response = FHttpServerResponse::Error(EHttpServerResponseCodes::NotSupported, StringOutDevice);
 		OnComplete(MoveTemp(Response));
 	}
 
diff --git a/Engine/Source/Runtime/PerfCounters/Private/PerfCountersModule.cpp b/Engine/Source/Runtime/PerfCounters/Private/PerfCountersModule.cpp
--- a/Engine/Source/Runtime/PerfCounters/Private/PerfCountersModule.cpp
+++ b/Engine/Source/Runtime/PerfCounters/Private/PerfCountersModule.cpp
@@ -4,6 +4,7 @@
 #include "HAL/PlatformProcess.h"
 #include "PerfCounters.h"
 #include "Misc/CommandLine.h"
+#include "Misc/Fork.h"
 
 class FPerfCountersModule : public IPerfCountersModule
 {
@@ -77,5 +78,13 @@ int32 IPerfCountersModule::GetHTTPStatsPort()
 {
 	int32 StatsPort = -1;
 	FParse::Value(FCommandLine::Get(), TEXT("statsPort="), StatsPort);
+
+	if (StatsPort > 0 && FParse::Param(FCommandLine::Get(), TEXT("forkAutoStatsPort")))
+	{
+		const int32 ChildProcessIdx = FForkProcessHelper::GetForkedChildProcessIndex();
+		StatsPort += ChildProcessIdx;
+		UE_LOG(LogPerfCounters, Display, TEXT("ForkAutoStatsPort assigned port %d to child process %d"), StatsPort, ChildProcessIdx);
+	}
+
 	return StatsPort;
 }
\ No newline at end of file
