Fixed unsafe type cast warnings in PerfCounters module.

diff --git a/Engine/Source/Runtime/PerfCounters/Private/PerfCounters.cpp b/Engine/Source/Runtime/PerfCounters/Private/PerfCounters.cpp
--- a/Engine/Source/Runtime/PerfCounters/Private/PerfCounters.cpp
+++ b/Engine/Source/Runtime/PerfCounters/Private/PerfCounters.cpp
@@ -52,7 +52,7 @@ bool FPerfCounters::Initialize()
 	{
 		InternalCountersUpdateInterval = ConfigInternalCountersUpdateInterval;
 	}
-	LastTimeInternalCountersUpdated = FPlatformTime::Seconds() - InternalCountersUpdateInterval * FMath::FRand();	// randomize between servers
+	LastTimeInternalCountersUpdated = static_cast<float>(FPlatformTime::Seconds()) - InternalCountersUpdateInterval * FMath::FRand();	// randomize between servers
 
 	// get the requested port from the command line (if specified)
 	const int32 StatsPort = IPerfCountersModule::GetHTTPStatsPort();
@@ -190,7 +190,7 @@ void FPerfCounters::TickZeroLoad(float DeltaTime)
 void FPerfCounters::TickSystemCounters(float DeltaTime)
 {
 	// set some internal perf stats ([RCL] FIXME 2015-12-08: move to a better place)
-	float CurrentTime = FPlatformTime::Seconds();
+	float CurrentTime = static_cast<float>(FPlatformTime::Seconds());
 	if (CurrentTime - LastTimeInternalCountersUpdated > InternalCountersUpdateInterval)
 	{
 		// get CPU stats first
diff --git a/Engine/Source/Runtime/PerfCounters/Private/ZeroLoad.cpp b/Engine/Source/Runtime/PerfCounters/Private/ZeroLoad.cpp
--- a/Engine/Source/Runtime/PerfCounters/Private/ZeroLoad.cpp
+++ b/Engine/Source/Runtime/PerfCounters/Private/ZeroLoad.cpp
@@ -38,7 +38,7 @@ uint32 FZeroLoad::Run()
 	double PreviousTick = FPlatformTime::Seconds();
 	while (!ExitRequest.GetValue())
 	{
-		FPlatformProcess::SleepNoStats(SecondsToSleep);
+		FPlatformProcess::SleepNoStats(static_cast<float>(SecondsToSleep));
 
 		const double CurrentTick = FPlatformTime::Seconds();
 		const double TickDuration = CurrentTick - PreviousTick;
