Fix crash when invalid command is sent to stats port.

diff --git a/Engine/Source/Runtime/PerfCounters/Private/PerfCounters.cpp b/Engine/Source/Runtime/PerfCounters/Private/PerfCounters.cpp
--- a/Engine/Source/Runtime/PerfCounters/Private/PerfCounters.cpp
+++ b/Engine/Source/Runtime/PerfCounters/Private/PerfCounters.cpp
@@ -233,16 +233,25 @@ bool FPerfCounters::ProcessStatsRequest(const FHttpServerRequest& Request, const
 
 bool FPerfCounters::ProcessExecRequest(const FHttpServerRequest& Request, const FHttpResultCallback& OnComplete)
 {
-	FStringOutputDevice StringOutDevice;
-	bool bExecCommandSuccess = false;
 	const FString* ExecCmd = Request.QueryParams.Find(TEXT("C"));
 	if (ExecCmd)
 	{
-		StringOutDevice.SetAutoEmitLineTerminator(true);
-
 		if (ExecCmdCallback.IsBound())
 		{
-			bExecCommandSuccess = ExecCmdCallback.Execute(*ExecCmd, StringOutDevice);
+			FStringOutputDevice StringOutDevice;
+			StringOutDevice.SetAutoEmitLineTerminator(true);
+
+			const bool bExecCommandSuccess = ExecCmdCallback.Execute(*ExecCmd, StringOutDevice);
+			if (bExecCommandSuccess)
+			{
+				auto Response = FHttpServerResponse::Create(StringOutDevice, TEXT("text/text"));
+				OnComplete(MoveTemp(Response));
+			}
+			else
+			{
+				auto Response = FHttpServerResponse::Error(EHttpServerResponseCodes::NotSupported, StringOutDevice);
+				OnComplete(MoveTemp(Response));
+			}
 		}
 		else
 		{
@@ -258,18 +267,6 @@ bool FPerfCounters::ProcessExecRequest(const FHttpServerRequest& Request, const
 		OnComplete(MoveTemp(Response));
 	}
 
-
-	if (bExecCommandSuccess)
-	{
-		auto Response = FHttpServerResponse::Create(StringOutDevice, TEXT("text/text"));
-		OnComplete(MoveTemp(Response));
-	}
-	else
-	{
-		auto Response = FHttpServerResponse::Error(EHttpServerResponseCodes::NotSupported, StringOutDevice);
-		OnComplete(MoveTemp(Response));
-	}
-
 	return true;
 }
 
